Added TechAudioTools plugin

diff --git a/Engine/Plugins/Experimental/TechAudioTools/Source/TechAudioToolsMetaSound/Private/TechAudioToolsMetaSoundModule.cpp b/Engine/Plugins/Experimental/TechAudioTools/Source/TechAudioToolsMetaSound/Private/TechAudioToolsMetaSoundModule.cpp
new file mode 100644
--- /dev/null
+++ b/Engine/Plugins/Experimental/TechAudioTools/Source/TechAudioToolsMetaSound/Private/TechAudioToolsMetaSoundModule.cpp
@@ -0,0 +1,5 @@
+﻿// Copyright Epic Games, Inc. All Rights Reserved.
+
+#include "Modules/ModuleManager.h"
+
+IMPLEMENT_MODULE(FDefaultModuleImpl, TechAudioToolsMetaSound)
diff --git a/Engine/Plugins/Experimental/TechAudioTools/Source/TechAudioToolsMetaSound/Private/ViewModels/MetaSoundViewModel.cpp b/Engine/Plugins/Experimental/TechAudioTools/Source/TechAudioToolsMetaSound/Private/ViewModels/MetaSoundViewModel.cpp
new file mode 100644
--- /dev/null
+++ b/Engine/Plugins/Experimental/TechAudioTools/Source/TechAudioToolsMetaSound/Private/ViewModels/MetaSoundViewModel.cpp
@@ -0,0 +1,346 @@
+﻿// Copyright Epic Games, Inc. All Rights Reserved.
+
+#include "ViewModels/MetaSoundViewModel.h"
+
+#include "Logging/StructuredLog.h"
+#include "MetasoundDocumentBuilderRegistry.h"
+#include "MetasoundOutputSubsystem.h"
+
+#include UE_INLINE_GENERATED_CPP_BY_NAME(MetaSoundViewModel)
+
+DEFINE_LOG_CATEGORY(LogTechAudioToolsMetaSound);
+
+namespace TechAudioTools::MetaSoundViewModel
+{
+	FName GetAdjustedDataType(const FName& CurrentDataType, const bool bIsArray)
+	{
+		const FString DataTypeString = CurrentDataType.ToString();
+		if (bIsArray)
+		{
+			// Append the suffix if it isn't already there
+			return DataTypeString.EndsWith(TEXT(":Array")) ? CurrentDataType : FName(DataTypeString + TEXT(":Array"));
+		}
+		else
+		{
+			// Split on ':' and keep the base data type
+			FString BaseDataType;
+			DataTypeString.Split(TEXT(":"), &BaseDataType, nullptr);
+			return FName(BaseDataType);
+		}
+	}
+} // TechAudioTools::MetaSoundViewModel
+
+FText UMetaSoundViewModel::GetBuilderNameAsText() const
+{
+	return FText::FromString(GetNameSafe(Builder));
+}
+
+TArray<UMetaSoundInputViewModel*> UMetaSoundViewModel::GetInputViewModels() const
+{
+	TArray<TObjectPtr<UMetaSoundInputViewModel>> InputViewModelsArray;
+	InputViewModels.GenerateValueArray(InputViewModelsArray);
+	return InputViewModelsArray;
+}
+
+TArray<UMetaSoundOutputViewModel*> UMetaSoundViewModel::GetOutputViewModels() const
+{
+	TArray<TObjectPtr<UMetaSoundOutputViewModel>> OutputViewModelsArray;
+	OutputViewModels.GenerateValueArray(OutputViewModelsArray);
+	return OutputViewModelsArray;
+}
+
+void UMetaSoundViewModel::InitializeMetaSound(const TScriptInterface<IMetaSoundDocumentInterface> InMetaSound)
+{
+	if (!InMetaSound)
+	{
+		Reset();
+		return;
+	}
+
+	Builder = Metasound::Engine::FDocumentBuilderRegistry::GetChecked().FindBuilderObject(InMetaSound);
+	Initialize(Builder);
+}
+
+void UMetaSoundViewModel::Initialize(UMetaSoundBuilderBase* InBuilder)
+{
+	Reset();
+
+	if (!InBuilder)
+	{
+		UE_LOGFMT(LogTechAudioToolsMetaSound, Log, "Unable to initialize MetaSoundViewModel. Builder was null.");
+		return;
+	}
+
+	Builder = InBuilder;
+
+	const FMetaSoundFrontendDocumentBuilder& DocumentBuilder = InBuilder->GetConstBuilder();
+	const FMetasoundFrontendDocument& FrontendDocument = DocumentBuilder.GetConstDocumentChecked();
+
+	InitializeProperties(FrontendDocument);
+
+	CreateMemberViewModels();
+	SetIsInitialized(true);
+}
+
+void UMetaSoundViewModel::Reset()
+{
+	SetIsInitialized(false);
+	ResetProperties();
+	Builder = nullptr;
+
+	InputViewModels.Empty();
+	UE_MVVM_BROADCAST_FIELD_VALUE_CHANGED(GetInputViewModels);
+
+	OutputViewModels.Empty();
+	UE_MVVM_BROADCAST_FIELD_VALUE_CHANGED(GetOutputViewModels);
+}
+
+void UMetaSoundViewModel::InitializeProperties(const FMetasoundFrontendDocument& FrontendDocument)
+{
+	UE_MVVM_SET_PROPERTY_VALUE(bIsPreset, FrontendDocument.RootGraph.PresetOptions.bIsPreset);
+}
+
+void UMetaSoundViewModel::ResetProperties()
+{
+	UE_MVVM_SET_PROPERTY_VALUE(bIsPreset, false);
+}
+
+void UMetaSoundViewModel::CreateMemberViewModels()
+{
+	const FMetaSoundFrontendDocumentBuilder& DocumentBuilder = Builder->GetConstBuilder();
+	const FMetasoundFrontendDocument& FrontendDocument = DocumentBuilder.GetConstDocumentChecked();
+
+	for (const FMetasoundFrontendClassInput& Input : FrontendDocument.RootGraph.GetDefaultInterface().Inputs)
+	{
+		CreateInputViewModel(Input);
+	}
+	UE_MVVM_BROADCAST_FIELD_VALUE_CHANGED(GetInputViewModels);
+
+	for (const FMetasoundFrontendClassOutput& Output : FrontendDocument.RootGraph.GetDefaultInterface().Outputs)
+	{
+		CreateOutputViewModel(Output);
+	}
+	UE_MVVM_BROADCAST_FIELD_VALUE_CHANGED(GetOutputViewModels);
+}
+
+void UMetaSoundViewModel::CreateInputViewModel(const FMetasoundFrontendClassInput& InInput)
+{
+	UMetaSoundInputViewModel* InputViewModel = NewObject<UMetaSoundInputViewModel>(this);
+	check(InputViewModel);
+
+	UMetaSoundBuilderBase* BuilderPtr = Builder.Get();
+	InputViewModel->InitializeInput(BuilderPtr, InInput);
+	InputViewModels.Emplace(InInput.Name, InputViewModel);
+
+	// Prevent spamming field notify broadcasts while we're still initializing.
+	if (bIsInitialized)
+	{
+		UE_MVVM_BROADCAST_FIELD_VALUE_CHANGED(GetInputViewModels);
+	}
+}
+
+void UMetaSoundViewModel::CreateOutputViewModel(const FMetasoundFrontendClassOutput& InOutput)
+{
+	UMetaSoundOutputViewModel* OutputViewModel = NewObject<UMetaSoundOutputViewModel>(this);
+	check(OutputViewModel);
+
+	UMetaSoundBuilderBase* BuilderPtr = Builder.Get();
+	OutputViewModel->InitializeOutput(BuilderPtr, InOutput);
+	OutputViewModels.Emplace(InOutput.Name, OutputViewModel);
+
+	if (bIsInitialized)
+	{
+		UE_MVVM_BROADCAST_FIELD_VALUE_CHANGED(GetOutputViewModels);
+	}
+}
+
+void UMetaSoundInputViewModel::InitializeInput(UMetaSoundBuilderBase* InBuilder, const FMetasoundFrontendClassInput& InInput)
+{
+	ResetInput();
+	Builder = InBuilder;
+
+	UE_MVVM_SET_PROPERTY_VALUE(InputName, InInput.Name);
+	UE_MVVM_SET_PROPERTY_VALUE(DataType, InInput.TypeName);
+
+	const FGuid PageID = Metasound::Engine::FDocumentBuilderRegistry::GetChecked().ResolveTargetPageID(InInput);
+	const FMetasoundFrontendLiteral& DefaultLiteral = InInput.FindConstDefaultChecked(PageID);
+
+	UE_MVVM_SET_PROPERTY_VALUE(LiteralType, DefaultLiteral.GetType());
+	UE_MVVM_SET_PROPERTY_VALUE(Literal, DefaultLiteral);
+	UE_MVVM_SET_PROPERTY_VALUE(bIsArray, DefaultLiteral.IsArray());
+
+	SetIsInitialized(true);
+}
+
+void UMetaSoundInputViewModel::ResetInput()
+{
+	SetIsInitialized(false);
+	Builder = nullptr;
+
+	UE_MVVM_SET_PROPERTY_VALUE(InputName, FName());
+	UE_MVVM_SET_PROPERTY_VALUE(DataType, FName());
+	UE_MVVM_SET_PROPERTY_VALUE(LiteralType, EMetasoundFrontendLiteralType::Invalid);
+	UE_MVVM_SET_PROPERTY_VALUE(Literal, FMetasoundFrontendLiteral());
+	UE_MVVM_SET_PROPERTY_VALUE(bIsArray, false);
+}
+
+void UMetaSoundInputViewModel::SetIsInitialized(const bool bInIsInitialized)
+{
+	UE_MVVM_SET_PROPERTY_VALUE(bIsInitialized, bInIsInitialized);
+}
+
+void UMetaSoundInputViewModel::SetInputName(const FName& InInputName)
+{
+	if (!Builder)
+	{
+		UE_LOGFMT(LogTechAudioToolsMetaSound, Log, "Failed to set input name {new} for {old}. Builder was null.", InInputName, InputName);
+		return;
+	}
+
+	const FName OldName = InputName;
+	if (UE_MVVM_SET_PROPERTY_VALUE(InputName, InInputName))
+	{
+		EMetaSoundBuilderResult Result;
+		Builder->SetGraphInputName(OldName, InInputName, Result);
+
+		if (Result != EMetaSoundBuilderResult::Succeeded)
+		{
+			UE_LOGFMT(LogTechAudioToolsMetaSound, Log, "Failed to set input name {new} for {old}. Builder was null.", InInputName, InputName);
+		}
+	}
+}
+
+void UMetaSoundInputViewModel::SetDataType(const FName& InDataType)
+{
+	if (!Builder)
+	{
+		UE_LOGFMT(LogTechAudioToolsMetaSound, Log, "Failed to set data type for {input}. Builder was null.", InputName);
+		return;
+	}
+
+	FMetaSoundFrontendDocumentBuilder& DocumentBuilder = Builder->GetBuilder();
+	if (!DocumentBuilder.IsValid())
+	{
+		UE_LOGFMT(LogTechAudioToolsMetaSound, Log, "Failed to set data type for {input}. DocumentBuilder was invalid.", InputName);
+		return;
+	}
+
+	if (UE_MVVM_SET_PROPERTY_VALUE(DataType, InDataType))
+	{
+		DocumentBuilder.SetGraphInputDataType(InputName, InDataType);
+	}
+}
+
+void UMetaSoundInputViewModel::SetIsArray(const bool bInIsArray)
+{
+	if (UE_MVVM_SET_PROPERTY_VALUE(bIsArray, bInIsArray))
+	{
+		SetDataType(TechAudioTools::MetaSoundViewModel::GetAdjustedDataType(DataType, bInIsArray));
+	}
+}
+
+FText UMetaSoundInputViewModel::GetLiteralValueAsText() const
+{
+	return FText::FromString(Literal.ToString());
+}
+
+void UMetaSoundInputViewModel::SetLiteral(const FMetasoundFrontendLiteral& InLiteral)
+{
+	if (!Builder)
+	{
+		UE_LOGFMT(LogTechAudioToolsMetaSound, Log, "Failed to set literal for {input}. Builder was null.", InputName);
+		return;
+	}
+
+	if (UE_MVVM_SET_PROPERTY_VALUE(Literal, InLiteral))
+	{
+		EMetaSoundBuilderResult Result;
+		Builder->SetGraphInputDefault(InputName, InLiteral, Result);
+		UE_MVVM_BROADCAST_FIELD_VALUE_CHANGED(GetLiteralValueAsText);
+	}
+}
+
+void UMetaSoundOutputViewModel::InitializeOutput(UMetaSoundBuilderBase* InBuilder, const FMetasoundFrontendClassOutput& InOutput)
+{
+	ResetOutput();
+	Builder = InBuilder;
+
+	UE_MVVM_SET_PROPERTY_VALUE(OutputName, InOutput.Name);
+	UE_MVVM_SET_PROPERTY_VALUE(DataType, InOutput.TypeName);
+	SetMetaSoundOutput(FMetaSoundOutput());
+	SetIsInitialized(true);
+}
+
+void UMetaSoundOutputViewModel::ResetOutput()
+{
+	SetMetaSoundOutput(FMetaSoundOutput());
+}
+
+void UMetaSoundOutputViewModel::SetIsInitialized(const bool bInIsInitialized)
+{
+	UE_MVVM_SET_PROPERTY_VALUE(bIsInitialized, bInIsInitialized);
+}
+
+void UMetaSoundOutputViewModel::SetOutputName(const FName& InOutputName)
+{
+	if (!Builder)
+	{
+		UE_LOGFMT(LogTechAudioToolsMetaSound, Log, "Failed to set output name for {output}. Builder was null.", OutputName);
+		return;
+	}
+
+	const FName OldName = OutputName;
+	if (UE_MVVM_SET_PROPERTY_VALUE(OutputName, InOutputName))
+	{
+		EMetaSoundBuilderResult Result;
+		Builder->SetGraphOutputName(OldName, InOutputName, Result);
+
+		if (Result != EMetaSoundBuilderResult::Succeeded)
+		{
+			UE_LOGFMT(LogTechAudioToolsMetaSound, Log, "Failed to set output name for {output}", OutputName);
+		}
+	}
+}
+
+void UMetaSoundOutputViewModel::SetDataType(const FName& InDataType)
+{
+	if (!Builder)
+	{
+		UE_LOGFMT(LogTechAudioToolsMetaSound, Log, "Failed to set data type for {output}. Builder was null.", OutputName);
+		return;
+	}
+
+	FMetaSoundFrontendDocumentBuilder& DocumentBuilder = Builder->GetBuilder();
+	if (!DocumentBuilder.IsValid())
+	{
+		UE_LOGFMT(LogTechAudioToolsMetaSound, Log, "Failed to set data type for {output}. DocumentBuilder was invalid.", OutputName);
+		return;
+	}
+
+	if (UE_MVVM_SET_PROPERTY_VALUE(DataType, InDataType))
+	{
+		DocumentBuilder.SetGraphOutputDataType(OutputName, InDataType);
+	}
+}
+
+void UMetaSoundOutputViewModel::SetIsArray(const bool bInIsArray)
+{
+	if (UE_MVVM_SET_PROPERTY_VALUE(bIsArray, bInIsArray))
+	{
+		SetDataType(TechAudioTools::MetaSoundViewModel::GetAdjustedDataType(DataType, bInIsArray));
+	}
+}
+
+void UMetaSoundOutputViewModel::SetMetaSoundOutput(const FMetaSoundOutput& InMetaSoundOutput)
+{
+	MetaSoundOutput = InMetaSoundOutput;
+	UE_MVVM_BROADCAST_FIELD_VALUE_CHANGED(MetaSoundOutput);
+}
+
+void UMetaSoundOutputViewModel::OnOutputValueChanged(const FName InOutputName, const FMetaSoundOutput& InOutput)
+{
+	if (OutputName == InOutputName)
+	{
+		SetMetaSoundOutput(InOutput);
+	}
+}
diff --git a/Engine/Plugins/Experimental/TechAudioTools/Source/TechAudioToolsMetaSound/Private/ViewModels/MetaSoundViewModelConversionFunctions.cpp b/Engine/Plugins/Experimental/TechAudioTools/Source/TechAudioToolsMetaSound/Private/ViewModels/MetaSoundViewModelConversionFunctions.cpp
new file mode 100644
--- /dev/null
+++ b/Engine/Plugins/Experimental/TechAudioTools/Source/TechAudioToolsMetaSound/Private/ViewModels/MetaSoundViewModelConversionFunctions.cpp
@@ -0,0 +1,33 @@
+﻿// Copyright Epic Games, Inc. All Rights Reserved.
+
+#include "ViewModels/MetaSoundViewModelConversionFunctions.h"
+
+#include "ViewModels/MetaSoundViewModel.h"
+
+#include UE_INLINE_GENERATED_CPP_BY_NAME(MetaSoundViewModelConversionFunctions)
+
+UMetaSoundInputViewModel* UMetaSoundViewModelConversionFunctions::FindInputViewModelByName(const TArray<UMetaSoundInputViewModel*>& MetaSoundInputViewModels, const FName InputName)
+{
+	for (UMetaSoundInputViewModel* InputViewModel : MetaSoundInputViewModels)
+	{
+		if (InputViewModel->GetInputName() == InputName)
+		{
+			return InputViewModel;
+		}
+	}
+
+	return nullptr;
+}
+
+UMetaSoundOutputViewModel* UMetaSoundViewModelConversionFunctions::FindOutputViewModelByName(const TArray<UMetaSoundOutputViewModel*>& MetaSoundOutputViewModels, const FName OutputName)
+{
+	for (UMetaSoundOutputViewModel* OutputViewModel : MetaSoundOutputViewModels)
+	{
+		if (OutputViewModel->GetOutputName() == OutputName)
+		{
+			return OutputViewModel;
+		}
+	}
+
+	return nullptr;
+}
diff --git a/Engine/Plugins/Experimental/TechAudioTools/Source/TechAudioToolsMetaSound/Public/ViewModels/MetaSoundViewModel.h b/Engine/Plugins/Experimental/TechAudioTools/Source/TechAudioToolsMetaSound/Public/ViewModels/MetaSoundViewModel.h
new file mode 100644
--- /dev/null
+++ b/Engine/Plugins/Experimental/TechAudioTools/Source/TechAudioToolsMetaSound/Public/ViewModels/MetaSoundViewModel.h
@@ -0,0 +1,193 @@
+// Copyright Epic Games, Inc. All Rights Reserved.
+
+#pragma once
+
+#include "MetasoundFrontendDocumentBuilder.h"
+#include "MetasoundGeneratorHandle.h"
+#include "MVVMViewModelBase.h"
+
+#include "MetaSoundViewModel.generated.h"
+
+TECHAUDIOTOOLSMETASOUND_API DECLARE_LOG_CATEGORY_EXTERN(LogTechAudioToolsMetaSound, Log, All);
+
+class UMetaSoundBuilderBase;
+class UMetaSoundInputViewModel;
+class UMetaSoundOutputViewModel;
+
+/**
+ * The base class for MetaSound viewmodels. Used for binding metadata and member inputs/outputs of a MetaSound to widgets in UMG.
+ * Can be initialized using a MetaSound Builder or a MetaSound asset. Creates member viewmodels for each input and output in the
+ * MetaSound upon initialization.
+ */
+UCLASS(DisplayName = "MetaSound Viewmodel")
+class TECHAUDIOTOOLSMETASOUND_API UMetaSoundViewModel : public UMVVMViewModelBase
+{
+	GENERATED_BODY()
+
+private:
+	// True if this MetaSound Viewmodel has been initialized.
+	UPROPERTY(BlueprintReadOnly, FieldNotify, Category = "MetaSound Viewmodel", meta = (AllowPrivateAccess))
+	bool bIsInitialized = false;
+
+	// True if the initialized MetaSound is a preset.
+	UPROPERTY(BlueprintReadOnly, FieldNotify, Category = "MetaSound Viewmodel", meta = (AllowPrivateAccess))
+	bool bIsPreset = false;
+
+public:
+	// Returns the object name of the initialized builder as text.
+	UFUNCTION(BlueprintCallable, FieldNotify, Category = "MetaSound Viewmodel")
+	FText GetBuilderNameAsText() const;
+
+	// Contains MetaSound Input Viewmodels for each input of the initialized MetaSound.
+	UFUNCTION(BlueprintCallable, FieldNotify, DisplayName = "Get Input Viewmodels", Category = "MetaSound Viewmodel")
+	virtual TArray<UMetaSoundInputViewModel*> GetInputViewModels() const;
+
+	// Contains MetaSound Output ViewModels for each output of the initialized MetaSound.
+	UFUNCTION(BlueprintCallable, FieldNotify, DisplayName = "Get Output Viewmodels", Category = "MetaSound Viewmodel")
+	virtual TArray<UMetaSoundOutputViewModel*> GetOutputViewModels() const;
+
+	// Initializes the viewmodel using the given MetaSound asset.
+	UFUNCTION(BlueprintCallable, DisplayName = "Initialize MetaSound", Category = "Audio|MetaSound Viewmodel")
+	virtual void InitializeMetaSound(const TScriptInterface<IMetaSoundDocumentInterface> InMetaSound);
+
+	// Initializes the viewmodel using the given builder.
+	UFUNCTION(BlueprintCallable, DisplayName = "Initialize Builder", Category = "Audio|MetaSound Viewmodel")
+	virtual void Initialize(UMetaSoundBuilderBase* InBuilder);
+
+	// Resets this MetaSoundViewModel instance to an uninitialized state.
+	UFUNCTION(BlueprintCallable, Category = "Audio|MetaSound Viewmodel")
+	virtual void Reset();
+
+	// Returns a reference to the initialized MetaSound's Builder.
+	UFUNCTION(BlueprintCallable, Category = "Audio|MetaSound Viewmodel")
+	UMetaSoundBuilderBase* GetBuilder() const { return Builder; }
+
+	bool IsInitialized() const { return bIsInitialized; }
+	void SetIsInitialized(const bool bInIsInitialized) { UE_MVVM_SET_PROPERTY_VALUE(bIsInitialized, bInIsInitialized); }
+
+protected:
+	virtual void InitializeProperties(const FMetasoundFrontendDocument& FrontendDocument);
+	virtual void ResetProperties();
+
+	// Called upon initialization. Creates viewmodel instances for all inputs and outputs of the initialized MetaSound.
+	void CreateMemberViewModels();
+
+	// Creates a single MetaSoundInputViewModel instance for the given input.
+	virtual void CreateInputViewModel(const FMetasoundFrontendClassInput& InInput);
+
+	// Creates a single MetaSoundOutputViewModel instance for the given output.
+	void CreateOutputViewModel(const FMetasoundFrontendClassOutput& InOutput);
+
+	UPROPERTY(Transient)
+	TObjectPtr<UMetaSoundBuilderBase> Builder;
+
+	UPROPERTY(Transient)
+	TMap<FName, TObjectPtr<UMetaSoundInputViewModel>> InputViewModels;
+
+	UPROPERTY(Transient)
+	TMap<FName, TObjectPtr<UMetaSoundOutputViewModel>> OutputViewModels;
+};
+
+/**
+ * Viewmodel class for MetaSound inputs. Allows widgets in UMG to bind to MetaSound literals. Useful for creating knobs, sliders, and other
+ * widgets for setting MetaSound input parameters.
+ */
+UCLASS(DisplayName = "MetaSound Input Viewmodel")
+class TECHAUDIOTOOLSMETASOUND_API UMetaSoundInputViewModel : public UMVVMViewModelBase
+{
+	GENERATED_BODY()
+
+protected:
+	// True if this MetaSoundInputViewModel has been initialized.
+	UPROPERTY(BlueprintReadOnly, FieldNotify, Category = "MetaSound Input", meta = (AllowPrivateAccess))
+	bool bIsInitialized = false;
+
+	// Sets the name of the initialized MetaSound input.
+	UPROPERTY(BlueprintReadWrite, EditAnywhere, FieldNotify, Setter, Category = "MetaSound Input", meta = (AllowPrivateAccess))
+	FName InputName;
+
+	// Returns the data type of the initialized MetaSound input.
+	UPROPERTY(BlueprintReadWrite, FieldNotify, Setter, Category = "MetaSound Input", meta = (AllowPrivateAccess))
+	FName DataType;
+
+	// True if the initialized MetaSound input is an array.
+	UPROPERTY(BlueprintReadWrite, FieldNotify, Setter = "SetIsArray", Category = "MetaSound Input", meta = (AllowPrivateAccess))
+	bool bIsArray;
+
+	// The MetaSound Literal belonging to the initialized MetaSound input.
+	UPROPERTY(BlueprintReadWrite, FieldNotify, Setter, Category = "MetaSound Input", meta = (AllowPrivateAccess))
+	FMetasoundFrontendLiteral Literal;
+
+	// Returns the Literal Type belonging to the initialized MetaSound input.
+	UPROPERTY(BlueprintReadOnly, FieldNotify, Category = "MetaSound Input", meta = (AllowPrivateAccess))
+	EMetasoundFrontendLiteralType LiteralType;
+
+	UPROPERTY(Transient)
+	TObjectPtr<UMetaSoundBuilderBase> Builder;
+
+public:
+	virtual void InitializeInput(UMetaSoundBuilderBase* InBuilder, const FMetasoundFrontendClassInput& InInput);
+	virtual void ResetInput();
+
+	FName GetInputName() const { return InputName; }
+
+	void SetIsInitialized(const bool bInIsInitialized);
+	void SetInputName(const FName& InInputName);
+	void SetDataType(const FName& InDataType);
+	void SetIsArray(const bool bInIsArray);
+
+	// Returns the value of this input's MetaSound Literal as a text value.
+	UFUNCTION(BlueprintCallable, FieldNotify, Category = "Audio|MetaSound Input Viewmodel")
+	FText GetLiteralValueAsText() const;
+
+	void SetLiteral(const FMetasoundFrontendLiteral& InLiteral);
+};
+
+/**
+ * Viewmodel class for MetaSound outputs. Allows widgets in UMG to bind to data from a MetaSound output. Useful for driving visual parameters
+ * using MetaSound outputs.
+ */
+UCLASS(DisplayName = "MetaSound Output Viewmodel")
+class TECHAUDIOTOOLSMETASOUND_API UMetaSoundOutputViewModel : public UMVVMViewModelBase
+{
+	GENERATED_BODY()
+
+protected:
+	// True if this MetaSoundOutputViewModel has been initialized.
+	UPROPERTY(BlueprintReadOnly, FieldNotify, Category = "MetaSound Output", meta = (AllowPrivateAccess))
+	bool bIsInitialized = false;
+
+	// Sets the name of the initialized MetaSound output.
+	UPROPERTY(BlueprintReadWrite, EditAnywhere, FieldNotify, Setter, Category = "MetaSound Output", meta = (AllowPrivateAccess))
+	FName OutputName;
+
+	// Returns the data type of the initialized MetaSound output.
+	UPROPERTY(BlueprintReadWrite, FieldNotify, Setter, Category = "MetaSound Output", meta = (AllowPrivateAccess))
+	FName DataType;
+
+	// True if the initialized MetaSound output is an array.
+	UPROPERTY(BlueprintReadWrite, FieldNotify, Setter = "SetIsArray", Category = "MetaSound Output", meta = (AllowPrivateAccess))
+	bool bIsArray;
+
+	// The MetaSound Output belonging to the initialized MetaSound output.
+	UPROPERTY(BlueprintReadWrite, FieldNotify, Setter, DisplayName = "MetaSound Output", Category = "MetaSound Output", meta = (AllowPrivateAccess))
+	FMetaSoundOutput MetaSoundOutput;
+
+	UPROPERTY(Transient)
+	TObjectPtr<UMetaSoundBuilderBase> Builder;
+
+public:
+	virtual void InitializeOutput(UMetaSoundBuilderBase* InBuilder, const FMetasoundFrontendClassOutput& InOutput);
+	virtual void ResetOutput();
+
+	FName GetOutputName() const { return OutputName; }
+
+	void SetIsInitialized(const bool bInIsInitialized);
+	void SetOutputName(const FName& InOutputName);
+	void SetDataType(const FName& InDataType);
+	void SetIsArray(const bool bInIsArray);
+	void SetMetaSoundOutput(const FMetaSoundOutput& InMetaSoundOutput);
+
+	UFUNCTION()
+	void OnOutputValueChanged(FName InOutputName, const FMetaSoundOutput& InOutput);
+};
diff --git a/Engine/Plugins/Experimental/TechAudioTools/Source/TechAudioToolsMetaSound/Public/ViewModels/MetaSoundViewModelConversionFunctions.h b/Engine/Plugins/Experimental/TechAudioTools/Source/TechAudioToolsMetaSound/Public/ViewModels/MetaSoundViewModelConversionFunctions.h
new file mode 100644
--- /dev/null
+++ b/Engine/Plugins/Experimental/TechAudioTools/Source/TechAudioToolsMetaSound/Public/ViewModels/MetaSoundViewModelConversionFunctions.h
@@ -0,0 +1,28 @@
+﻿// Copyright Epic Games, Inc. All Rights Reserved.
+
+#pragma once
+
+#include "Kismet/BlueprintFunctionLibrary.h"
+
+#include "MetaSoundViewModelConversionFunctions.generated.h"
+
+class UMetaSoundInputViewModel;
+class UMetaSoundOutputViewModel;
+
+/**
+ * Collection of conversion functions to use with MetaSound Viewmodels.
+ */
+UCLASS(MinimalAPI)
+class UMetaSoundViewModelConversionFunctions : public UBlueprintFunctionLibrary
+{
+	GENERATED_BODY()
+
+public:
+	// Searches the given array of MetaSoundInputViewModels for the specified input. Returns nullptr if unable to find a matching viewmodel.
+	UFUNCTION(BlueprintCallable, BlueprintPure, DisplayName = "Find Input Viewmodel by Name", Category = "TechAudioTools")
+	static UMetaSoundInputViewModel* FindInputViewModelByName(const TArray<UMetaSoundInputViewModel*>& MetaSoundInputViewModels, const FName InputName);
+
+	// Searches the given array of MetaSoundOutputViewModels for the specified output. Returns nullptr if unable to find a matching viewmodel.
+	UFUNCTION(BlueprintCallable, BlueprintPure, DisplayName = "Find Output Viewmodel by Name", Category = "TechAudioTools")
+	static UMetaSoundOutputViewModel* FindOutputViewModelByName(const TArray<UMetaSoundOutputViewModel*>& MetaSoundOutputViewModels, const FName OutputName);
+};
diff --git a/Engine/Plugins/Experimental/TechAudioTools/Source/TechAudioToolsMetaSound/TechAudioToolsMetaSound.Build.cs b/Engine/Plugins/Experimental/TechAudioTools/Source/TechAudioToolsMetaSound/TechAudioToolsMetaSound.Build.cs
new file mode 100644
--- /dev/null
+++ b/Engine/Plugins/Experimental/TechAudioTools/Source/TechAudioToolsMetaSound/TechAudioToolsMetaSound.Build.cs
@@ -0,0 +1,29 @@
+﻿// Copyright Epic Games, Inc. All Rights Reserved.
+
+using UnrealBuildTool;
+
+public class TechAudioToolsMetaSound : ModuleRules
+{
+    public TechAudioToolsMetaSound(ReadOnlyTargetRules Target) : base(Target)
+    {
+        PCHUsage = ModuleRules.PCHUsageMode.UseExplicitOrSharedPCHs;
+
+        PublicDependencyModuleNames.AddRange(
+            new string[]
+            {
+                "Core",
+				"MetasoundEngine",
+				"MetasoundFrontend",
+				"ModelViewViewModel",
+            }
+        );
+
+        PrivateDependencyModuleNames.AddRange(
+            new string[]
+            {
+                "CoreUObject",
+                "Engine",
+            }
+        );
+    }
+}
diff --git a/Engine/Plugins/Experimental/TechAudioTools/Source/TechAudioToolsMetaSoundEditor/Private/TechAudioToolsMetaSoundEditorModule.cpp b/Engine/Plugins/Experimental/TechAudioTools/Source/TechAudioToolsMetaSoundEditor/Private/TechAudioToolsMetaSoundEditorModule.cpp
new file mode 100644
--- /dev/null
+++ b/Engine/Plugins/Experimental/TechAudioTools/Source/TechAudioToolsMetaSoundEditor/Private/TechAudioToolsMetaSoundEditorModule.cpp
@@ -0,0 +1,5 @@
+﻿// Copyright Epic Games, Inc. All Rights Reserved.
+
+#include "Modules/ModuleManager.h"
+
+IMPLEMENT_MODULE(FDefaultModuleImpl, TechAudioToolsMetaSoundEditor)
diff --git a/Engine/Plugins/Experimental/TechAudioTools/Source/TechAudioToolsMetaSoundEditor/Private/ViewModels/MetaSoundEditorViewModel.cpp b/Engine/Plugins/Experimental/TechAudioTools/Source/TechAudioToolsMetaSoundEditor/Private/ViewModels/MetaSoundEditorViewModel.cpp
new file mode 100644
--- /dev/null
+++ b/Engine/Plugins/Experimental/TechAudioTools/Source/TechAudioToolsMetaSoundEditor/Private/ViewModels/MetaSoundEditorViewModel.cpp
@@ -0,0 +1,450 @@
+﻿// Copyright Epic Games, Inc. All Rights Reserved.
+
+#include "ViewModels/MetaSoundEditorViewModel.h"
+
+#include "Editor/EditorEngine.h"
+#include "Logging/StructuredLog.h"
+#include "MetasoundDocumentBuilderRegistry.h"
+#include "MetasoundEditorBuilderListener.h"
+#include "MetasoundEditorSubsystem.h"
+#include "MetasoundOutputSubsystem.h"
+
+#include UE_INLINE_GENERATED_CPP_BY_NAME(MetaSoundEditorViewModel)
+
+extern UNREALED_API UEditorEngine* GEditor;
+
+DEFINE_LOG_CATEGORY(LogTechAudioToolsMetaSoundEditor);
+
+bool UMetaSoundEditorViewModel::IsEditorOnly() const
+{
+	return true;
+}
+
+UWorld* UMetaSoundEditorViewModel::GetWorld() const
+{
+	if (HasAllFlags(RF_ClassDefaultObject))
+	{
+		return nullptr;
+	}
+
+	return GEditor ? GEditor->GetEditorWorldContext(false).World() : nullptr;
+}
+
+void UMetaSoundEditorViewModel::InitializeMetaSound(const TScriptInterface<IMetaSoundDocumentInterface> InMetaSound)
+{
+	if (!InMetaSound)
+	{
+		Reset();
+		return;
+	}
+
+	Builder = &Metasound::Engine::FDocumentBuilderRegistry::GetChecked().FindOrBeginBuilding(*InMetaSound.GetObject());
+	Initialize(Builder);
+}
+
+void UMetaSoundEditorViewModel::Initialize(UMetaSoundBuilderBase* InBuilder)
+{
+	Super::Initialize(InBuilder);
+	BindDelegates();
+}
+
+void UMetaSoundEditorViewModel::Reset()
+{
+	Super::Reset();
+	UnbindDelegates();
+}
+
+void UMetaSoundEditorViewModel::InitializeProperties(const FMetasoundFrontendDocument& FrontendDocument)
+{
+	Super::InitializeProperties(FrontendDocument);
+
+	UE_MVVM_SET_PROPERTY_VALUE(DisplayName, FrontendDocument.RootGraph.Metadata.GetDisplayName());
+	UE_MVVM_SET_PROPERTY_VALUE(Description, FrontendDocument.RootGraph.Metadata.GetDescription());
+	UE_MVVM_SET_PROPERTY_VALUE(Author, FrontendDocument.RootGraph.Metadata.GetAuthor());
+	UE_MVVM_SET_PROPERTY_VALUE(Keywords, FrontendDocument.RootGraph.Metadata.GetKeywords());
+	UE_MVVM_SET_PROPERTY_VALUE(CategoryHierarchy, FrontendDocument.RootGraph.Metadata.GetCategoryHierarchy());
+	UE_MVVM_SET_PROPERTY_VALUE(bIsDeprecated, FrontendDocument.RootGraph.Metadata.GetIsDeprecated());
+}
+
+void UMetaSoundEditorViewModel::ResetProperties()
+{
+	Super::ResetProperties();
+
+	UE_MVVM_SET_PROPERTY_VALUE(DisplayName, FText());
+	UE_MVVM_SET_PROPERTY_VALUE(Description, FText());
+	UE_MVVM_SET_PROPERTY_VALUE(Author, FString());
+	UE_MVVM_SET_PROPERTY_VALUE(Keywords, TArray<FText>());
+	UE_MVVM_SET_PROPERTY_VALUE(CategoryHierarchy, TArray<FText>());
+	UE_MVVM_SET_PROPERTY_VALUE(bIsDeprecated, false);
+}
+
+void UMetaSoundEditorViewModel::SetMetaSoundDisplayName(const FText& InDisplayName)
+{
+	if (Builder && UE_MVVM_SET_PROPERTY_VALUE(DisplayName, InDisplayName))
+	{
+		FMetaSoundFrontendDocumentBuilder& DocumentBuilder = Builder->GetBuilder();
+		DocumentBuilder.SetDisplayName(InDisplayName);
+	}
+}
+
+void UMetaSoundEditorViewModel::SetMetaSoundDescription(const FText& InDescription)
+{
+	if (Builder && UE_MVVM_SET_PROPERTY_VALUE(Description, InDescription))
+	{
+		FMetaSoundFrontendDocumentBuilder& DocumentBuilder = Builder->GetBuilder();
+		DocumentBuilder.SetDescription(InDescription);
+	}
+}
+
+void UMetaSoundEditorViewModel::SetAuthor(const FString& InAuthor)
+{
+	if (Builder && UE_MVVM_SET_PROPERTY_VALUE(Author, InAuthor))
+	{
+		FMetaSoundFrontendDocumentBuilder& DocumentBuilder = Builder->GetBuilder();
+		DocumentBuilder.SetAuthor(InAuthor);
+	}
+}
+
+void UMetaSoundEditorViewModel::SetKeywords(const TArray<FText>& InKeywords)
+{
+	if (Builder && UE_MVVM_SET_PROPERTY_VALUE(Keywords, InKeywords))
+	{
+		FMetaSoundFrontendDocumentBuilder& DocumentBuilder = Builder->GetBuilder();
+		DocumentBuilder.SetKeywords(InKeywords);
+	}
+}
+
+void UMetaSoundEditorViewModel::SetCategoryHierarchy(const TArray<FText>& InCategoryHierarchy)
+{
+	if (Builder && UE_MVVM_SET_PROPERTY_VALUE(CategoryHierarchy, InCategoryHierarchy))
+	{
+		FMetaSoundFrontendDocumentBuilder& DocumentBuilder = Builder->GetBuilder();
+		DocumentBuilder.SetCategoryHierarchy(InCategoryHierarchy);
+	}
+}
+
+void UMetaSoundEditorViewModel::SetIsDeprecated(const bool bInIsDeprecated)
+{
+	if (Builder && UE_MVVM_SET_PROPERTY_VALUE(bIsDeprecated, bInIsDeprecated))
+	{
+		FMetaSoundFrontendDocumentBuilder& DocumentBuilder = Builder->GetBuilder();
+		DocumentBuilder.SetIsDeprecated(bInIsDeprecated);
+	}
+}
+
+void UMetaSoundEditorViewModel::CreateInputViewModel(const FMetasoundFrontendClassInput& InInput)
+{
+	Super::CreateInputViewModel(InInput);
+
+	UMetaSoundInputViewModel* InputViewModel = *InputViewModels.Find(InInput.Name);
+	if (UMetaSoundInputEditorViewModel* InputEditorViewModel = Cast<UMetaSoundInputEditorViewModel>(InputViewModel))
+	{
+		if (BuilderListener && InputEditorViewModel)
+		{
+			BuilderListener->OnGraphInputDefaultChangedDelegate.AddDynamic(InputEditorViewModel, &UMetaSoundInputEditorViewModel::OnInputDefaultChanged);
+		}
+	}
+}
+
+void UMetaSoundEditorViewModel::BindDelegates()
+{
+	if (!Builder)
+	{
+		UE_LOGFMT(LogTechAudioToolsMetaSoundEditor, Log, "Could not bind MetaSoundViewModel delegates. Builder was null.");
+		SetIsInitialized(false);
+		return;
+	}
+
+	UMetaSoundEditorSubsystem* EditorSubsystem = GEditor->GetEditorSubsystem<UMetaSoundEditorSubsystem>();
+	if (!EditorSubsystem)
+	{
+		UE_LOGFMT(LogTechAudioToolsMetaSoundEditor, Log, "Could not bind MetaSoundViewModel delegates. Unable to locate MetaSound Editor Subsystem.");
+		SetIsInitialized(false);
+		return;
+	}
+
+	EMetaSoundBuilderResult Result;
+	BuilderListener = EditorSubsystem->AddBuilderDelegateListener(Builder, Result);
+
+	if (Result != EMetaSoundBuilderResult::Succeeded)
+	{
+		UE_LOGFMT(LogTechAudioToolsMetaSoundEditor, Log, "Could not bind MetaSoundViewModel delegates. Failed to create BuilderListener.");
+		SetIsInitialized(false);
+		return;
+	}
+
+	if (BuilderListener)
+	{
+		BuilderListener->OnGraphInputAddedDelegate.AddDynamic(this, &ThisClass::OnInputAdded);
+		BuilderListener->OnRemovingGraphInputDelegate.AddDynamic(this, &ThisClass::OnInputRemoved);
+		BuilderListener->OnGraphInputNameChangedDelegate.AddDynamic(this, &ThisClass::OnInputNameChanged);
+		BuilderListener->OnGraphInputDataTypeChangedDelegate.AddDynamic(this, &ThisClass::OnInputDataTypeChanged);
+
+		BuilderListener->OnGraphOutputAddedDelegate.AddDynamic(this, &ThisClass::OnOutputAdded);
+		BuilderListener->OnRemovingGraphOutputDelegate.AddDynamic(this, &ThisClass::OnOutputRemoved);
+		BuilderListener->OnGraphOutputNameChangedDelegate.AddDynamic(this, &ThisClass::OnOutputNameChanged);
+		BuilderListener->OnGraphOutputDataTypeChangedDelegate.AddDynamic(this, &ThisClass::OnOutputDataTypeChanged);
+	}
+}
+
+void UMetaSoundEditorViewModel::UnbindDelegates()
+{
+	if (BuilderListener)
+	{
+		BuilderListener->RemoveAllDelegates();
+		BuilderListener = nullptr;
+	}
+}
+
+void UMetaSoundEditorViewModel::OnInputAdded(const FName VertexName, const FName DataType)
+{
+	const FMetaSoundFrontendDocumentBuilder& DocumentBuilder = Builder->GetConstBuilder();
+	const FMetasoundFrontendDocument& FrontendDocument = DocumentBuilder.GetConstDocumentChecked();
+	for (const FMetasoundFrontendClassInput& Input : FrontendDocument.RootGraph.GetDefaultInterface().Inputs)
+	{
+		if (Input.Name == VertexName)
+		{
+			CreateInputViewModel(Input);
+			return;
+		}
+	}
+}
+
+void UMetaSoundEditorViewModel::OnInputRemoved(const FName VertexName, const FName DataType)
+{
+	InputViewModels.Remove(VertexName);
+	UE_MVVM_BROADCAST_FIELD_VALUE_CHANGED(GetInputViewModels);
+}
+
+void UMetaSoundEditorViewModel::OnInputNameChanged(const FName OldName, const FName NewName)
+{
+	if (UMetaSoundInputViewModel* InputViewModel = *InputViewModels.Find(OldName))
+	{
+		InputViewModel->SetInputName(NewName);
+		UE_MVVM_BROADCAST_FIELD_VALUE_CHANGED(GetInputViewModels);
+	}
+}
+
+void UMetaSoundEditorViewModel::OnInputDataTypeChanged(const FName VertexName, const FName DataType)
+{
+	if (UMetaSoundInputViewModel* InputViewModel = *InputViewModels.Find(VertexName))
+	{
+		InputViewModel->SetDataType(DataType);
+		UE_MVVM_BROADCAST_FIELD_VALUE_CHANGED(GetInputViewModels);
+	}
+}
+
+void UMetaSoundEditorViewModel::OnOutputAdded(const FName VertexName, const FName DataType)
+{
+	if (Builder)
+	{
+		const FMetaSoundFrontendDocumentBuilder& DocumentBuilder = Builder->GetConstBuilder();
+		const FMetasoundFrontendDocument& FrontendDocument = DocumentBuilder.GetConstDocumentChecked();
+		for (const FMetasoundFrontendClassOutput& Output : FrontendDocument.RootGraph.GetDefaultInterface().Outputs)
+		{
+			if (Output.Name == VertexName)
+			{
+				CreateOutputViewModel(Output);
+				return;
+			}
+		}
+	}
+}
+
+void UMetaSoundEditorViewModel::OnOutputRemoved(const FName VertexName, const FName DataType)
+{
+	OutputViewModels.Remove(VertexName);
+	UE_MVVM_BROADCAST_FIELD_VALUE_CHANGED(GetOutputViewModels);
+}
+
+void UMetaSoundEditorViewModel::OnOutputNameChanged(const FName OldName, const FName NewName)
+{
+	if (UMetaSoundOutputViewModel* OutputViewModel = *OutputViewModels.Find(OldName))
+	{
+		OutputViewModel->SetOutputName(NewName);
+		UE_MVVM_BROADCAST_FIELD_VALUE_CHANGED(GetOutputViewModels);
+	}
+}
+
+void UMetaSoundEditorViewModel::OnOutputDataTypeChanged(const FName VertexName, const FName DataType)
+{
+	if (UMetaSoundOutputViewModel* OutputViewModel = *OutputViewModels.Find(VertexName))
+	{
+		OutputViewModel->SetDataType(DataType);
+		UE_MVVM_BROADCAST_FIELD_VALUE_CHANGED(GetOutputViewModels);
+	}
+}
+
+bool UMetaSoundInputEditorViewModel::IsEditorOnly() const
+{
+	return true;
+}
+
+UWorld* UMetaSoundInputEditorViewModel::GetWorld() const
+{
+	if (HasAllFlags(RF_ClassDefaultObject))
+	{
+		return nullptr;
+	}
+
+	return GEditor ? GEditor->GetEditorWorldContext(false).World() : nullptr;
+}
+
+void UMetaSoundInputEditorViewModel::InitializeInput(UMetaSoundBuilderBase* InBuilder, const FMetasoundFrontendClassInput& InInput)
+{
+	Super::InitializeInput(InBuilder, InInput);
+
+	UE_MVVM_SET_PROPERTY_VALUE(InputDisplayName, InInput.Metadata.GetDisplayName());
+	UE_MVVM_SET_PROPERTY_VALUE(InputDescription, InInput.Metadata.GetDescription());
+	UE_MVVM_SET_PROPERTY_VALUE(SortOrderIndex, InInput.Metadata.SortOrderIndex);
+	UE_MVVM_SET_PROPERTY_VALUE(bIsAdvancedDisplay, InInput.Metadata.bIsAdvancedDisplay);
+}
+
+void UMetaSoundInputEditorViewModel::ResetInput()
+{
+	Super::ResetInput();
+
+	UE_MVVM_SET_PROPERTY_VALUE(InputDisplayName, FText());
+	UE_MVVM_SET_PROPERTY_VALUE(InputDescription, FText());
+	UE_MVVM_SET_PROPERTY_VALUE(SortOrderIndex, 0);
+	UE_MVVM_SET_PROPERTY_VALUE(bIsAdvancedDisplay, false);
+}
+
+void UMetaSoundInputEditorViewModel::SetInputDisplayName(const FText& InDisplayName)
+{
+	if (Builder)
+	{
+		FMetaSoundFrontendDocumentBuilder& DocumentBuilder = Builder->GetBuilder();
+		if (DocumentBuilder.IsValid() && UE_MVVM_SET_PROPERTY_VALUE(InputDisplayName, InDisplayName))
+		{
+			DocumentBuilder.SetGraphOutputDisplayName(InputName, InDisplayName);
+		}
+	}
+}
+
+void UMetaSoundInputEditorViewModel::SetInputDescription(const FText& InDescription)
+{
+	if (Builder)
+	{
+		FMetaSoundFrontendDocumentBuilder& DocumentBuilder = Builder->GetBuilder();
+		if (DocumentBuilder.IsValid() && UE_MVVM_SET_PROPERTY_VALUE(InputDescription, InDescription))
+		{
+			DocumentBuilder.SetGraphInputDescription(InputName, InDescription);
+		}
+	}
+}
+
+void UMetaSoundInputEditorViewModel::SetSortOrderIndex(const int32 InSortOrderIndex)
+{
+	if (Builder)
+	{
+		FMetaSoundFrontendDocumentBuilder& DocumentBuilder = Builder->GetBuilder();
+		if (DocumentBuilder.IsValid() && UE_MVVM_SET_PROPERTY_VALUE(SortOrderIndex, InSortOrderIndex))
+		{
+			DocumentBuilder.SetGraphInputSortOrderIndex(InputName, InSortOrderIndex);
+		}
+	}
+}
+
+void UMetaSoundInputEditorViewModel::SetIsAdvancedDisplay(const bool bInIsAdvancedDisplay)
+{
+	if (Builder)
+	{
+		FMetaSoundFrontendDocumentBuilder& DocumentBuilder = Builder->GetBuilder();
+		if (DocumentBuilder.IsValid() && UE_MVVM_SET_PROPERTY_VALUE(bIsAdvancedDisplay, bInIsAdvancedDisplay))
+		{
+			DocumentBuilder.SetGraphInputAdvancedDisplay(InputName, bInIsAdvancedDisplay);
+		}
+	}
+}
+
+void UMetaSoundInputEditorViewModel::OnInputDefaultChanged(const FName VertexName, FMetasoundFrontendLiteral LiteralValue, FName PageName)
+{
+	if (InputName == VertexName)
+	{
+		SetLiteral(LiteralValue);
+	}
+}
+
+bool UMetaSoundOutputEditorViewModel::IsEditorOnly() const
+{
+	return true;
+}
+
+UWorld* UMetaSoundOutputEditorViewModel::GetWorld() const
+{
+	if (HasAllFlags(RF_ClassDefaultObject))
+	{
+		return nullptr;
+	}
+
+	return GEditor ? GEditor->GetEditorWorldContext(false).World() : nullptr;
+}
+
+void UMetaSoundOutputEditorViewModel::InitializeOutput(UMetaSoundBuilderBase* InBuilder, const FMetasoundFrontendClassOutput& InOutput)
+{
+	Super::InitializeOutput(InBuilder, InOutput);
+
+	UE_MVVM_SET_PROPERTY_VALUE(OutputDisplayName, InOutput.Metadata.GetDisplayName());
+	UE_MVVM_SET_PROPERTY_VALUE(OutputDescription, InOutput.Metadata.GetDescription());
+	UE_MVVM_SET_PROPERTY_VALUE(SortOrderIndex, InOutput.Metadata.SortOrderIndex);
+	UE_MVVM_SET_PROPERTY_VALUE(bIsAdvancedDisplay, InOutput.Metadata.bIsAdvancedDisplay);
+}
+
+void UMetaSoundOutputEditorViewModel::ResetOutput()
+{
+	Super::ResetOutput();
+
+	UE_MVVM_SET_PROPERTY_VALUE(OutputDisplayName, FText());
+	UE_MVVM_SET_PROPERTY_VALUE(OutputDescription, FText());
+	UE_MVVM_SET_PROPERTY_VALUE(SortOrderIndex, 0);
+	UE_MVVM_SET_PROPERTY_VALUE(bIsAdvancedDisplay, false);
+}
+
+void UMetaSoundOutputEditorViewModel::SetOutputDisplayName(const FText& InDisplayName)
+{
+	if (Builder)
+	{
+		FMetaSoundFrontendDocumentBuilder& DocumentBuilder = Builder->GetBuilder();
+		if (DocumentBuilder.IsValid() && UE_MVVM_SET_PROPERTY_VALUE(OutputDisplayName, InDisplayName))
+		{
+			DocumentBuilder.SetGraphOutputDisplayName(OutputName, InDisplayName);
+		}
+	}
+}
+
+void UMetaSoundOutputEditorViewModel::SetOutputDescription(const FText& InDescription)
+{
+	if (Builder)
+	{
+		FMetaSoundFrontendDocumentBuilder& DocumentBuilder = Builder->GetBuilder();
+		if (DocumentBuilder.IsValid() && UE_MVVM_SET_PROPERTY_VALUE(OutputDescription, InDescription))
+		{
+			DocumentBuilder.SetGraphOutputDescription(OutputName, InDescription);
+		}
+	}
+}
+
+void UMetaSoundOutputEditorViewModel::SetSortOrderIndex(const int32 InSortOrderIndex)
+{
+	if (Builder)
+	{
+		FMetaSoundFrontendDocumentBuilder& DocumentBuilder = Builder->GetBuilder();
+		if (DocumentBuilder.IsValid() && UE_MVVM_SET_PROPERTY_VALUE(SortOrderIndex, InSortOrderIndex))
+		{
+			DocumentBuilder.SetGraphOutputSortOrderIndex(OutputName, InSortOrderIndex);
+		}
+	}
+}
+
+void UMetaSoundOutputEditorViewModel::SetIsAdvancedDisplay(const bool bInIsAdvancedDisplay)
+{
+	if (Builder)
+	{
+		FMetaSoundFrontendDocumentBuilder& DocumentBuilder = Builder->GetBuilder();
+		if (DocumentBuilder.IsValid() && UE_MVVM_SET_PROPERTY_VALUE(bIsAdvancedDisplay, bInIsAdvancedDisplay))
+		{
+			DocumentBuilder.SetGraphOutputAdvancedDisplay(OutputName, bInIsAdvancedDisplay);
+		}
+	}
+}
diff --git a/Engine/Plugins/Experimental/TechAudioTools/Source/TechAudioToolsMetaSoundEditor/Private/ViewModels/MetaSoundEditorViewModelConversionFunctions.cpp b/Engine/Plugins/Experimental/TechAudioTools/Source/TechAudioToolsMetaSoundEditor/Private/ViewModels/MetaSoundEditorViewModelConversionFunctions.cpp
new file mode 100644
--- /dev/null
+++ b/Engine/Plugins/Experimental/TechAudioTools/Source/TechAudioToolsMetaSoundEditor/Private/ViewModels/MetaSoundEditorViewModelConversionFunctions.cpp
@@ -0,0 +1,43 @@
+﻿// Copyright Epic Games, Inc. All Rights Reserved.
+
+#include "ViewModels/MetaSoundEditorViewModelConversionFunctions.h"
+
+#include "Editor/EditorEngine.h"
+#include "MetasoundEditorGraphSchema.h"
+#include "MetasoundEditorModule.h"
+#include "Modules/ModuleManager.h"
+
+#include UE_INLINE_GENERATED_CPP_BY_NAME(MetaSoundEditorViewModelConversionFunctions)
+
+extern UNREALED_API UEditorEngine* GEditor;
+
+bool UMetaSoundEditorViewModelConversionFunctions::IsEditorOnly() const
+{
+	return true;
+}
+
+UWorld* UMetaSoundEditorViewModelConversionFunctions::GetWorld() const
+{
+	if (HasAllFlags(RF_ClassDefaultObject))
+	{
+		return nullptr;
+	}
+
+	return GEditor ? GEditor->GetEditorWorldContext(false).World() : nullptr;
+}
+
+FLinearColor UMetaSoundEditorViewModelConversionFunctions::GetMetaSoundDataTypePinColor(const FName& DataType)
+{
+	using namespace Metasound::Editor;
+
+	const IMetasoundEditorModule& MetaSoundEditorModule = FModuleManager::GetModuleChecked<IMetasoundEditorModule>("MetaSoundEditor");
+	if (const FEdGraphPinType* PinType = MetaSoundEditorModule.FindPinType(DataType))
+	{
+		if (const UMetasoundEditorGraphSchema* Schema = GetDefault<UMetasoundEditorGraphSchema>())
+		{
+			return Schema->GetPinTypeColor(*PinType);
+		}
+	}
+
+	return FLinearColor::Black;
+}
diff --git a/Engine/Plugins/Experimental/TechAudioTools/Source/TechAudioToolsMetaSoundEditor/Public/ViewModels/MetaSoundEditorViewModel.h b/Engine/Plugins/Experimental/TechAudioTools/Source/TechAudioToolsMetaSoundEditor/Public/ViewModels/MetaSoundEditorViewModel.h
new file mode 100644
--- /dev/null
+++ b/Engine/Plugins/Experimental/TechAudioTools/Source/TechAudioToolsMetaSoundEditor/Public/ViewModels/MetaSoundEditorViewModel.h
@@ -0,0 +1,195 @@
+﻿// Copyright Epic Games, Inc. All Rights Reserved.
+
+#pragma once
+
+#include "ViewModels/MetaSoundViewModel.h"
+
+#include "MetaSoundEditorViewModel.generated.h"
+
+class UMetaSoundEditorBuilderListener;
+
+TECHAUDIOTOOLSMETASOUNDEDITOR_API DECLARE_LOG_CATEGORY_EXTERN(LogTechAudioToolsMetaSoundEditor, Log, All);
+
+/**
+ * Editor viewmodel for MetaSounds. Creates MetaSoundEditorBuilderListener bindings upon initialization, allowing changes made to assets in the
+ * MetaSound Editor to be reflected in UMG widgets.
+ */
+UCLASS(MinimalAPI, DisplayName = "MetaSound Editor Viewmodel")
+class UMetaSoundEditorViewModel : public UMetaSoundViewModel
+{
+	GENERATED_BODY()
+
+protected:
+	// Sets the display name of the initialized MetaSound.
+	UPROPERTY(BlueprintReadWrite, FieldNotify, Setter = "SetMetaSoundDisplayName", Category = "MetaSound Viewmodel|Metadata", meta = (AllowPrivateAccess))
+	FText DisplayName;
+
+	// Sets the description of the initialized MetaSound.
+	UPROPERTY(BlueprintReadWrite, FieldNotify, Setter = "SetMetaSoundDescription", Category = "MetaSound Viewmodel|Metadata", meta = (AllowPrivateAccess))
+	FText Description;
+
+	// Sets the author of the initialized MetaSound.
+	UPROPERTY(BlueprintReadWrite, FieldNotify, Setter, Category = "MetaSound Viewmodel|Metadata", meta = (AllowPrivateAccess))
+	FString Author;
+
+	// Sets the keywords of the initialized MetaSound.
+	UPROPERTY(BlueprintReadWrite, FieldNotify, Setter, Category = "MetaSound Viewmodel|Metadata", meta = (AllowPrivateAccess))
+	TArray<FText> Keywords;
+
+	// Sets the category hierarchy of the initialized MetaSound.
+	UPROPERTY(BlueprintReadWrite, FieldNotify, Setter, Category = "MetaSound Viewmodel|Metadata", meta = (AllowPrivateAccess))
+	TArray<FText> CategoryHierarchy;
+
+	// Sets the initialized MetaSound asset as deprecated.
+	UPROPERTY(BlueprintReadWrite, FieldNotify, Setter = "SetIsDeprecated", Category = "MetaSound Viewmodel|Metadata", meta = (AllowPrivateAccess))
+	bool bIsDeprecated = false;
+
+public:
+	virtual bool IsEditorOnly() const override;
+	virtual UWorld* GetWorld() const override;
+
+	virtual void InitializeMetaSound(const TScriptInterface<IMetaSoundDocumentInterface> InMetaSound) override;
+	virtual void Initialize(UMetaSoundBuilderBase* InBuilder) override;
+	virtual void Reset() override;
+
+	virtual void InitializeProperties(const FMetasoundFrontendDocument& FrontendDocument) override;
+	virtual void ResetProperties() override;
+
+	void SetMetaSoundDisplayName(const FText& InDisplayName);
+	void SetMetaSoundDescription(const FText& InDescription);
+	void SetAuthor(const FString& InAuthor);
+	void SetKeywords(const TArray<FText>& InKeywords);
+	void SetCategoryHierarchy(const TArray<FText>& InCategoryHierarchy);
+	void SetIsDeprecated(const bool bInIsDeprecated);
+
+protected:
+	UPROPERTY(Transient)
+	TObjectPtr<UMetaSoundEditorBuilderListener> BuilderListener;
+
+	virtual void CreateInputViewModel(const FMetasoundFrontendClassInput& InInput) override;
+
+	void BindDelegates();
+	void UnbindDelegates();
+
+	// Called when a new MetaSound input has been added to the initialized MetaSound.
+	UFUNCTION()
+	void OnInputAdded(FName VertexName, FName DataType);
+
+	// Called when a MetaSound input has been removed from the initialized MetaSound.
+	UFUNCTION()
+	void OnInputRemoved(FName VertexName, FName DataType);
+
+	// Called when the name of an input on the initialized MetaSound has changed.
+	UFUNCTION()
+	void OnInputNameChanged(FName OldName, FName NewName);
+
+	// Called when an input's data type has changed.
+	UFUNCTION()
+	void OnInputDataTypeChanged(FName VertexName, FName DataType);
+
+	// Called when a new MetaSound output has been added to the initialized MetaSound.
+	UFUNCTION()
+	void OnOutputAdded(FName VertexName, FName DataType);
+
+	// Called when a MetaSound output has been removed from the initialized MetaSound.
+	UFUNCTION()
+	void OnOutputRemoved(FName VertexName, FName DataType);
+
+	// Called when the name of an output on the initialized MetaSound has changed.
+	UFUNCTION()
+	void OnOutputNameChanged(FName OldName, FName NewName);
+
+	// Called when an output's data type has changed.
+	UFUNCTION()
+	void OnOutputDataTypeChanged(FName VertexName, FName DataType);
+};
+
+/**
+ * Editor viewmodel class for MetaSound inputs. Extends the runtime MetaSoundInputViewModel with editor-only functionality. 
+ */
+UCLASS(MinimalAPI, DisplayName = "MetaSound Input Editor Viewmodel")
+class UMetaSoundInputEditorViewModel : public UMetaSoundInputViewModel
+{
+	GENERATED_BODY()
+
+protected:
+	// Sets the display name of the initialized MetaSound input.
+	UPROPERTY(BlueprintReadWrite, FieldNotify, Setter, Category = "MetaSound Input", meta = (AllowPrivateAccess))
+	FText InputDisplayName;
+
+	// Sets the description of the initialized MetaSound input.
+	UPROPERTY(BlueprintReadWrite, FieldNotify, Setter, Category = "MetaSound Input", meta = (AllowPrivateAccess))
+	FText InputDescription;
+
+	// Sets the sort order index of the initialized MetaSound input.
+	UPROPERTY(BlueprintReadWrite, FieldNotify, Setter, Category = "MetaSound Input", meta = (AllowPrivateAccess))
+	int32 SortOrderIndex;
+
+	// Sets whether the initialized MetaSound input should be located in the Advanced Display category.
+	UPROPERTY(BlueprintReadWrite, FieldNotify, Setter = "SetIsAdvancedDisplay", Category = "MetaSound Input", meta = (AllowPrivateAccess))
+	bool bIsAdvancedDisplay;
+
+public:
+	virtual bool IsEditorOnly() const override;
+	virtual UWorld* GetWorld() const override;
+
+	virtual void InitializeInput(UMetaSoundBuilderBase* InBuilder, const FMetasoundFrontendClassInput& InInput) override;
+	virtual void ResetInput() override;
+
+	FText GetInputDisplayName() const { return InputDisplayName; }
+	FText GetInputDescription() const { return InputDescription; }
+	int32 GetSortOrderIndex() const { return SortOrderIndex; }
+	bool IsAdvancedDisplay() const { return bIsAdvancedDisplay; }
+
+	void SetInputDisplayName(const FText& InDisplayName);
+	void SetInputDescription(const FText& InDescription);
+	void SetSortOrderIndex(const int32 InSortOrderIndex);
+	void SetIsAdvancedDisplay(const bool bInIsAdvancedDisplay);
+
+	// Called when the default value of an input has been changed on the initialized MetaSound.
+	UFUNCTION()
+	void OnInputDefaultChanged(FName VertexName, FMetasoundFrontendLiteral LiteralValue, FName PageName);
+};
+
+/**
+ * Editor viewmodel class for MetaSound outputs. Extends the runtime MetaSoundOutputViewModel with editor-only functionality. 
+ */
+UCLASS(MinimalAPI, DisplayName = "MetaSound Output Editor Viewmodel")
+class UMetaSoundOutputEditorViewModel : public UMetaSoundOutputViewModel
+{
+	GENERATED_BODY()
+
+protected:
+	// Sets the display name of the initialized MetaSound output.
+	UPROPERTY(BlueprintReadWrite, FieldNotify, Setter, Category = "MetaSound Output", meta = (AllowPrivateAccess))
+	FText OutputDisplayName;
+
+	// Sets the description of the initialized MetaSound output.
+	UPROPERTY(BlueprintReadWrite, FieldNotify, Setter, Category = "MetaSound Output", meta = (AllowPrivateAccess))
+	FText OutputDescription;
+
+	// Sets the sort order index of the initialized MetaSound output.
+	UPROPERTY(BlueprintReadWrite, FieldNotify, Setter, Category = "MetaSound Output", meta = (AllowPrivateAccess))
+	int32 SortOrderIndex;
+
+	// Sets whether the initialized MetaSound output should be located in the Advanced Display category.
+	UPROPERTY(BlueprintReadWrite, FieldNotify, Setter = "SetIsAdvancedDisplay", Category = "MetaSound Output", meta = (AllowPrivateAccess))
+	bool bIsAdvancedDisplay;
+
+public:
+	virtual bool IsEditorOnly() const override;
+	virtual UWorld* GetWorld() const override;
+
+	virtual void InitializeOutput(UMetaSoundBuilderBase* InBuilder, const FMetasoundFrontendClassOutput& InOutput) override;
+	virtual void ResetOutput() override;
+
+	FText GetOutputDisplayName() const { return OutputDisplayName; }
+	FText GetOutputDescription() const { return OutputDescription; }
+	int32 GetSortOrderIndex() const { return SortOrderIndex; }
+	bool IsAdvancedDisplay() const { return bIsAdvancedDisplay; }
+
+	void SetOutputDisplayName(const FText& InDisplayName);
+	void SetOutputDescription(const FText& InDescription);
+	void SetSortOrderIndex(const int32 InSortOrderIndex);
+	void SetIsAdvancedDisplay(const bool bInIsAdvancedDisplay);
+};
diff --git a/Engine/Plugins/Experimental/TechAudioTools/Source/TechAudioToolsMetaSoundEditor/Public/ViewModels/MetaSoundEditorViewModelConversionFunctions.h b/Engine/Plugins/Experimental/TechAudioTools/Source/TechAudioToolsMetaSoundEditor/Public/ViewModels/MetaSoundEditorViewModelConversionFunctions.h
new file mode 100644
--- /dev/null
+++ b/Engine/Plugins/Experimental/TechAudioTools/Source/TechAudioToolsMetaSoundEditor/Public/ViewModels/MetaSoundEditorViewModelConversionFunctions.h
@@ -0,0 +1,24 @@
+﻿// Copyright Epic Games, Inc. All Rights Reserved.
+
+#pragma once
+
+#include "Kismet/BlueprintFunctionLibrary.h"
+
+#include "MetaSoundEditorViewModelConversionFunctions.generated.h"
+
+/**
+ * Collection of conversion functions to use with MetaSound Editor Viewmodels.
+ */
+UCLASS(MinimalAPI)
+class UMetaSoundEditorViewModelConversionFunctions : public UBlueprintFunctionLibrary
+{
+	GENERATED_BODY()
+
+public:
+	virtual bool IsEditorOnly() const override;
+	virtual UWorld* GetWorld() const override;
+
+	// Returns the pin color associated with the given data type.
+	UFUNCTION(BlueprintCallable, BlueprintPure, DisplayName = "Get MetaSound Data Type Pin Color", Category = "TechAudioTools")
+	static FLinearColor GetMetaSoundDataTypePinColor(const FName& DataType);
+};
diff --git a/Engine/Plugins/Experimental/TechAudioTools/Source/TechAudioToolsMetaSoundEditor/TechAudioToolsMetaSoundEditor.Build.cs b/Engine/Plugins/Experimental/TechAudioTools/Source/TechAudioToolsMetaSoundEditor/TechAudioToolsMetaSoundEditor.Build.cs
new file mode 100644
--- /dev/null
+++ b/Engine/Plugins/Experimental/TechAudioTools/Source/TechAudioToolsMetaSoundEditor/TechAudioToolsMetaSoundEditor.Build.cs
@@ -0,0 +1,34 @@
+﻿// Copyright Epic Games, Inc. All Rights Reserved.
+
+using UnrealBuildTool;
+
+public class TechAudioToolsMetaSoundEditor : ModuleRules
+{
+    public TechAudioToolsMetaSoundEditor(ReadOnlyTargetRules Target) : base(Target)
+    {
+        PCHUsage = ModuleRules.PCHUsageMode.UseExplicitOrSharedPCHs;
+
+        PublicDependencyModuleNames.AddRange(
+            new string[]
+            {
+	            "Blutility",
+                "Core",
+				"MetasoundEditor",
+				"MetasoundEngine",
+				"MetasoundFrontend",
+				"ModelViewViewModel",
+				"ModelViewViewModelEditor",
+				"TechAudioToolsMetaSound",
+				"UnrealEd",
+            }
+        );
+
+        PrivateDependencyModuleNames.AddRange(
+            new string[]
+            {
+                "CoreUObject",
+                "Engine",
+            }
+        );
+    }
+}
diff --git a/Engine/Plugins/Experimental/TechAudioTools/TechAudioTools.uplugin b/Engine/Plugins/Experimental/TechAudioTools/TechAudioTools.uplugin
new file mode 100644
--- /dev/null
+++ b/Engine/Plugins/Experimental/TechAudioTools/TechAudioTools.uplugin
@@ -0,0 +1,39 @@
+{
+	"FileVersion": 3,
+	"Version": 1,
+	"VersionName": "1.0",
+	"FriendlyName": "TechAudioTools",
+	"Description": "A collection of audio-related tools and utilities.",
+	"Category": "Audio",
+	"CreatedBy": "Epic Games, Inc.",
+	"CreatedByURL": "https://epicgames.com",
+	"DocsURL": "",
+	"MarketplaceURL": "",
+	"SupportURL": "",
+	"CanContainContent": true,
+	"IsBetaVersion": true,
+	"IsExperimentalVersion": true,
+	"Installed": false,
+	"Modules": [
+		{
+			"Name": "TechAudioToolsMetaSound",
+			"Type": "Runtime",
+			"LoadingPhase": "Default"
+		},
+		{
+			"Name": "TechAudioToolsMetaSoundEditor",
+			"Type": "Editor",
+			"LoadingPhase": "Default"
+		}
+	],
+	"Plugins": [
+		{
+			"Name": "Metasound",
+			"Enabled": true
+		},
+		{
+			"Name": "ModelViewViewModel",
+			"Enabled": true
+		}
+	]
+}
\ No newline at end of file
diff --git a/Engine/Plugins/Runtime/Metasound/Source/MetasoundEditor/Private/MetasoundEditorBuilderListener.h b/Engine/Plugins/Runtime/Metasound/Source/MetasoundEditor/Private/MetasoundEditorBuilderListener.h
deleted file mode 100644
diff --git a/Engine/Plugins/Runtime/Metasound/Source/MetasoundEditor/Public/MetasoundEditorBuilderListener.h b/Engine/Plugins/Runtime/Metasound/Source/MetasoundEditor/Public/MetasoundEditorBuilderListener.h
new file mode 100644
--- /dev/null
+++ b/Engine/Plugins/Runtime/Metasound/Source/MetasoundEditor/Public/MetasoundEditorBuilderListener.h
@@ -0,0 +1,108 @@
+// Copyright Epic Games, Inc. All Rights Reserved.
+#pragma once
+#include "Delegates/DelegateCombinations.h"
+#include "MetasoundBuilderBase.h"
+#include "MetasoundDocumentInterface.h"
+#include "MetasoundFrontendDocumentModifyDelegates.h"
+
+#include "MetasoundEditorBuilderListener.generated.h"
+
+// Forward Declarations
+struct FMetaSoundNodeHandle;
+struct FMetasoundFrontendLiteral;
+
+// BP Delegates for builder changes 
+DECLARE_DYNAMIC_MULTICAST_DELEGATE_TwoParams(FOnMetaSoundBuilderGraphInterfaceMutate, FName, VertexName, FName, DataType);
+DECLARE_DYNAMIC_MULTICAST_DELEGATE_ThreeParams(FOnMetaSoundBuilderGraphLiteralMutate, FName, VertexName, FMetasoundFrontendLiteral, LiteralValue, FName, PageName);
+DECLARE_DYNAMIC_MULTICAST_DELEGATE_TwoParams(FOnMetaSoundBuilderGraphVertexRename, FName, OldName, FName, NewName);
+
+UCLASS(MinimalAPI, BlueprintType)
+class UMetaSoundEditorBuilderListener : public UObject
+{
+public: 
+	GENERATED_BODY()
+
+	virtual ~UMetaSoundEditorBuilderListener() {}
+	
+	void Init(const TWeakObjectPtr<UMetaSoundBuilderBase> InBuilder);
+
+	void OnGraphInputAdded(int32 Index);
+	void OnGraphInputDefaultChanged(int32 Index);
+	void OnRemovingGraphInput(int32 Index);
+
+	void OnGraphOutputAdded(int32 Index);
+	void OnRemovingGraphOutput(int32 Index);
+
+	void OnGraphInputDataTypeChanged(int32 Index);
+	void OnGraphOutputDataTypeChanged(int32 Index);
+	void OnGraphInputNameChanged(FName OldName, FName NewName);
+	void OnGraphOutputNameChanged(FName OldName, FName NewName);
+
+
+	UFUNCTION(BlueprintCallable, Category = "Audio|MetaSound|Builder|Editor|Delegates")
+	METASOUNDEDITOR_API void RemoveAllDelegates();
+
+	UPROPERTY(BlueprintAssignable, BlueprintCallable, Category = "Audio|MetaSound|Builder|Editor|Delegates")
+	FOnMetaSoundBuilderGraphInterfaceMutate OnGraphInputAddedDelegate;
+	
+	UPROPERTY(BlueprintAssignable, BlueprintCallable, Category = "Audio|MetaSound|Builder|Editor|Delegates")
+	FOnMetaSoundBuilderGraphLiteralMutate OnGraphInputDefaultChangedDelegate;
+
+	UPROPERTY(BlueprintAssignable, BlueprintCallable, Category = "Audio|MetaSound|Builder|Editor|Delegates")
+	FOnMetaSoundBuilderGraphInterfaceMutate OnRemovingGraphInputDelegate;
+	
+	UPROPERTY(BlueprintAssignable, BlueprintCallable, Category= "Audio|MetaSound|Builder|Editor|Delegates")
+	FOnMetaSoundBuilderGraphInterfaceMutate OnGraphOutputAddedDelegate;
+
+	UPROPERTY(BlueprintAssignable, BlueprintCallable, Category = "Audio|MetaSound|Builder|Editor|Delegates")
+	FOnMetaSoundBuilderGraphInterfaceMutate OnRemovingGraphOutputDelegate;
+
+	UPROPERTY(BlueprintAssignable, BlueprintCallable, Category = "Audio|MetaSound|Builder|Editor|Delegates")
+	FOnMetaSoundBuilderGraphInterfaceMutate OnGraphInputDataTypeChangedDelegate;
+
+	UPROPERTY(BlueprintAssignable, BlueprintCallable, Category = "Audio|MetaSound|Builder|Editor|Delegates")
+	FOnMetaSoundBuilderGraphInterfaceMutate OnGraphOutputDataTypeChangedDelegate;
+
+	UPROPERTY(BlueprintAssignable, BlueprintCallable, Category = "Audio|MetaSound|Builder|Editor|Delegates")
+	FOnMetaSoundBuilderGraphVertexRename OnGraphInputNameChangedDelegate;
+
+	UPROPERTY(BlueprintAssignable, BlueprintCallable, Category = "Audio|MetaSound|Builder|Editor|Delegates")
+	FOnMetaSoundBuilderGraphVertexRename OnGraphOutputNameChangedDelegate;
+
+private:
+	// Handles for document delegates
+	FDelegateHandle OnInputAddedHandle;
+	FDelegateHandle OnInputDefaultChangedHandle;
+	FDelegateHandle OnRemovingInputHandle;
+
+	FDelegateHandle OnOutputAddedHandle;
+	FDelegateHandle OnRemovingOutputHandle;
+
+	FDelegateHandle OnInputDataTypeChangedHandle;
+	FDelegateHandle OnOutputDataTypeChangedHandle;
+	FDelegateHandle OnInputNameChangedHandle;
+	FDelegateHandle OnOutputNameChangedHandle;
+
+	TWeakObjectPtr<UMetaSoundBuilderBase> Builder;
+
+	class FEditorBuilderListener : public Metasound::Frontend::IDocumentBuilderTransactionListener
+	{
+	public: 
+
+		FEditorBuilderListener() = default;
+		FEditorBuilderListener(TObjectPtr<UMetaSoundEditorBuilderListener> InParent)
+			: Parent(InParent)
+		{
+		}
+
+		virtual ~FEditorBuilderListener() = default;
+
+		// IDocumentBuilderTransactionListener
+		virtual void OnBuilderReloaded(Metasound::Frontend::FDocumentModifyDelegates& OutDelegates) override;
+
+	private: 
+		TObjectPtr<UMetaSoundEditorBuilderListener> Parent;
+	};
+
+	TSharedPtr<FEditorBuilderListener> BuilderListener;
+};
\ No newline at end of file
diff --git a/Engine/Plugins/Runtime/Metasound/Source/MetasoundFrontend/Private/MetasoundFrontendDocumentBuilder.cpp b/Engine/Plugins/Runtime/Metasound/Source/MetasoundFrontend/Private/MetasoundFrontendDocumentBuilder.cpp
--- a/Engine/Plugins/Runtime/Metasound/Source/MetasoundFrontend/Private/MetasoundFrontendDocumentBuilder.cpp
+++ b/Engine/Plugins/Runtime/Metasound/Source/MetasoundFrontend/Private/MetasoundFrontendDocumentBuilder.cpp
@@ -4964,6 +4964,38 @@ bool FMetaSoundFrontendDocumentBuilder::SetGraphInputName(FName InputName, FName
 }
 
 #if WITH_EDITORONLY_DATA
+bool FMetaSoundFrontendDocumentBuilder::SetGraphInputSortOrderIndex(const FName InputName, const int32 InSortOrderIndex)
+{
+	using namespace Metasound;
+
+	if (const int32* Index = DocumentCache->GetInterfaceCache().FindInputIndex(InputName))
+	{
+		FMetasoundFrontendDocument& Document = GetDocumentChecked();
+		FMetasoundFrontendClassInput& GraphInput = Document.RootGraph.GetDefaultInterface().Inputs[*Index];
+		GraphInput.Metadata.SortOrderIndex = InSortOrderIndex;
+		Document.Metadata.ModifyContext.AddMemberIDModified(GraphInput.NodeID);
+		return true;
+	}
+
+	return false;
+}
+
+bool FMetaSoundFrontendDocumentBuilder::SetGraphOutputSortOrderIndex(const FName OutputName, const int32 InSortOrderIndex)
+{
+	using namespace Metasound;
+
+	if (const int32* Index = DocumentCache->GetInterfaceCache().FindOutputIndex(OutputName))
+	{
+		FMetasoundFrontendDocument& Document = GetDocumentChecked();
+		FMetasoundFrontendClassOutput& GraphOutput = Document.RootGraph.GetDefaultInterface().Outputs[*Index];
+		GraphOutput.Metadata.SortOrderIndex = InSortOrderIndex;
+		Document.Metadata.ModifyContext.AddMemberIDModified(GraphOutput.NodeID);
+		return true;
+	}
+
+	return false;
+}
+
 bool FMetaSoundFrontendDocumentBuilder::SetGraphOutputAdvancedDisplay(const FName OutputName, const bool InAdvancedDisplay)
 {
 	FMetasoundFrontendDocument& Document = GetDocumentChecked();
@@ -5258,6 +5290,26 @@ void FMetaSoundFrontendDocumentBuilder::SetDisplayName(const FText& InDisplayNam
 	DocumentInterface->GetDocument().RootGraph.Metadata.SetDisplayName(InDisplayName);
 }
 
+void FMetaSoundFrontendDocumentBuilder::SetDescription(const FText& InDescription)
+{
+	DocumentInterface->GetDocument().RootGraph.Metadata.SetDescription(InDescription);
+}
+
+void FMetaSoundFrontendDocumentBuilder::SetKeywords(const TArray<FText>& InKeywords)
+{
+	DocumentInterface->GetDocument().RootGraph.Metadata.SetKeywords(InKeywords);
+}
+
+void FMetaSoundFrontendDocumentBuilder::SetCategoryHierarchy(const TArray<FText>& InCategoryHierarchy)
+{
+	DocumentInterface->GetDocument().RootGraph.Metadata.SetCategoryHierarchy(InCategoryHierarchy);
+}
+
+void FMetaSoundFrontendDocumentBuilder::SetIsDeprecated(const bool bInIsDeprecated)
+{
+	DocumentInterface->GetDocument().RootGraph.Metadata.SetIsDeprecated(bInIsDeprecated);
+}
+
 void FMetaSoundFrontendDocumentBuilder::SetMemberMetadata(UMetaSoundFrontendMemberMetadata& NewMetadata)
 {
 	check(NewMetadata.MemberID.IsValid());
diff --git a/Engine/Plugins/Runtime/Metasound/Source/MetasoundFrontend/Public/MetasoundFrontendDocumentBuilder.h b/Engine/Plugins/Runtime/Metasound/Source/MetasoundFrontend/Public/MetasoundFrontendDocumentBuilder.h
--- a/Engine/Plugins/Runtime/Metasound/Source/MetasoundFrontend/Public/MetasoundFrontendDocumentBuilder.h
+++ b/Engine/Plugins/Runtime/Metasound/Source/MetasoundFrontend/Public/MetasoundFrontendDocumentBuilder.h
@@ -496,6 +496,12 @@ public:
 	bool SetGraphInputName(FName InputName, FName InName);
 
 #if WITH_EDITORONLY_DATA
+	// Sets the graph input's sort order index. Returns true if found and set, false if not.
+	bool SetGraphInputSortOrderIndex(const FName InputName, const int32 InSortOrderIndex);
+
+	// Sets the graph output's sort order index. Returns true if found and set, false if not.
+	bool SetGraphOutputSortOrderIndex(const FName OutputName, const int32 InSortOrderIndex);
+
 	// Sets the given output`s IsAdvancedDisplay state. AdvancedDisplay pins are hidden in the node by default.
 	// returns true if state was changed.
 	bool SetGraphOutputAdvancedDisplay(const FName OutputName, const bool InAdvancedDisplay);
@@ -538,6 +544,10 @@ public:
 
 #if WITH_EDITOR
 	void SetDisplayName(const FText& InDisplayName);
+	void SetDescription(const FText& InDescription);
+	void SetKeywords(const TArray<FText>& InKeywords);
+	void SetCategoryHierarchy(const TArray<FText>& InCategoryHierarchy);
+	void SetIsDeprecated(const bool bInIsDeprecated);
 
 	void SetMemberMetadata(UMetaSoundFrontendMemberMetadata& NewMetadata);
 
