TechAudioTools - update metasound editor viewmodel to instantiate editor viewmodels for inputs/outputs

#rnx

diff --git a/Engine/Plugins/Experimental/TechAudioTools/Source/TechAudioToolsMetaSound/Private/ViewModels/MetaSoundViewModel.cpp b/Engine/Plugins/Experimental/TechAudioTools/Source/TechAudioToolsMetaSound/Private/ViewModels/MetaSoundViewModel.cpp
--- a/Engine/Plugins/Experimental/TechAudioTools/Source/TechAudioToolsMetaSound/Private/ViewModels/MetaSoundViewModel.cpp
+++ b/Engine/Plugins/Experimental/TechAudioTools/Source/TechAudioToolsMetaSound/Private/ViewModels/MetaSoundViewModel.cpp
@@ -125,11 +125,10 @@ void UMetaSoundViewModel::CreateMemberViewModels()
 
 void UMetaSoundViewModel::CreateInputViewModel(const FMetasoundFrontendClassInput& InInput)
 {
-	UMetaSoundInputViewModel* InputViewModel = NewObject<UMetaSoundInputViewModel>(this);
+	UMetaSoundInputViewModel* InputViewModel = CreateInputViewModelInstance();
 	check(InputViewModel);
 
-	UMetaSoundBuilderBase* BuilderPtr = Builder.Get();
-	InputViewModel->InitializeInput(BuilderPtr, InInput);
+	InputViewModel->InitializeInput(Builder.Get(), InInput);
 	InputViewModels.Emplace(InInput.Name, InputViewModel);
 
 	// Prevent spamming field notify broadcasts while we're still initializing.
@@ -139,13 +138,17 @@ void UMetaSoundViewModel::CreateInputViewModel(const FMetasoundFrontendClassInpu
 	}
 }
 
+UMetaSoundInputViewModel* UMetaSoundViewModel::CreateInputViewModelInstance()
+{
+	return NewObject<UMetaSoundInputViewModel>(this, UMetaSoundInputViewModel::StaticClass());
+}
+
 void UMetaSoundViewModel::CreateOutputViewModel(const FMetasoundFrontendClassOutput& InOutput)
 {
-	UMetaSoundOutputViewModel* OutputViewModel = NewObject<UMetaSoundOutputViewModel>(this);
+	UMetaSoundOutputViewModel* OutputViewModel = CreateOutputViewModelInstance();
 	check(OutputViewModel);
 
-	UMetaSoundBuilderBase* BuilderPtr = Builder.Get();
-	OutputViewModel->InitializeOutput(BuilderPtr, InOutput);
+	OutputViewModel->InitializeOutput(Builder.Get(), InOutput);
 	OutputViewModels.Emplace(InOutput.Name, OutputViewModel);
 
 	if (bIsInitialized)
@@ -154,6 +157,11 @@ void UMetaSoundViewModel::CreateOutputViewModel(const FMetasoundFrontendClassOut
 	}
 }
 
+UMetaSoundOutputViewModel* UMetaSoundViewModel::CreateOutputViewModelInstance()
+{
+	return NewObject<UMetaSoundOutputViewModel>(this, UMetaSoundOutputViewModel::StaticClass());
+}
+
 void UMetaSoundInputViewModel::InitializeInput(UMetaSoundBuilderBase* InBuilder, const FMetasoundFrontendClassInput& InInput)
 {
 	ResetInput();
diff --git a/Engine/Plugins/Experimental/TechAudioTools/Source/TechAudioToolsMetaSound/Public/ViewModels/MetaSoundViewModel.h b/Engine/Plugins/Experimental/TechAudioTools/Source/TechAudioToolsMetaSound/Public/ViewModels/MetaSoundViewModel.h
--- a/Engine/Plugins/Experimental/TechAudioTools/Source/TechAudioToolsMetaSound/Public/ViewModels/MetaSoundViewModel.h
+++ b/Engine/Plugins/Experimental/TechAudioTools/Source/TechAudioToolsMetaSound/Public/ViewModels/MetaSoundViewModel.h
@@ -74,9 +74,11 @@ protected:
 
 	// Creates a single MetaSoundInputViewModel instance for the given input.
 	virtual void CreateInputViewModel(const FMetasoundFrontendClassInput& InInput);
+    virtual UMetaSoundInputViewModel* CreateInputViewModelInstance();
 
 	// Creates a single MetaSoundOutputViewModel instance for the given output.
 	void CreateOutputViewModel(const FMetasoundFrontendClassOutput& InOutput);
+	virtual UMetaSoundOutputViewModel* CreateOutputViewModelInstance();
 
 	UPROPERTY(Transient)
 	TObjectPtr<UMetaSoundBuilderBase> Builder;
diff --git a/Engine/Plugins/Experimental/TechAudioTools/Source/TechAudioToolsMetaSoundEditor/Private/ViewModels/MetaSoundEditorViewModel.cpp b/Engine/Plugins/Experimental/TechAudioTools/Source/TechAudioToolsMetaSoundEditor/Private/ViewModels/MetaSoundEditorViewModel.cpp
--- a/Engine/Plugins/Experimental/TechAudioTools/Source/TechAudioToolsMetaSoundEditor/Private/ViewModels/MetaSoundEditorViewModel.cpp
+++ b/Engine/Plugins/Experimental/TechAudioTools/Source/TechAudioToolsMetaSoundEditor/Private/ViewModels/MetaSoundEditorViewModel.cpp
@@ -146,6 +146,16 @@ void UMetaSoundEditorViewModel::CreateInputViewModel(const FMetasoundFrontendCla
 	}
 }
 
+UMetaSoundInputViewModel* UMetaSoundEditorViewModel::CreateInputViewModelInstance()
+{
+    return NewObject<UMetaSoundInputEditorViewModel>(this, UMetaSoundInputEditorViewModel::StaticClass());
+}
+
+UMetaSoundOutputViewModel* UMetaSoundEditorViewModel::CreateOutputViewModelInstance()
+{
+    return NewObject<UMetaSoundOutputEditorViewModel>(this, UMetaSoundOutputEditorViewModel::StaticClass());
+}
+
 void UMetaSoundEditorViewModel::BindDelegates()
 {
 	if (!Builder)
@@ -184,6 +194,14 @@ void UMetaSoundEditorViewModel::BindDelegates()
 		BuilderListener->OnRemovingGraphOutputDelegate.AddDynamic(this, &ThisClass::OnOutputRemoved);
 		BuilderListener->OnGraphOutputNameChangedDelegate.AddDynamic(this, &ThisClass::OnOutputNameChanged);
 		BuilderListener->OnGraphOutputDataTypeChangedDelegate.AddDynamic(this, &ThisClass::OnOutputDataTypeChanged);
+
+		for (const TPair<FName, UMetaSoundInputViewModel*> InputViewModel : InputViewModels)
+		{
+			if (UMetaSoundInputEditorViewModel* InputEditorViewModel = Cast<UMetaSoundInputEditorViewModel>(InputViewModel.Value))
+			{
+				BuilderListener->OnGraphInputDefaultChangedDelegate.AddDynamic(InputEditorViewModel, &UMetaSoundInputEditorViewModel::OnInputDefaultChanged);
+			}
+		}
 	}
 }
 
diff --git a/Engine/Plugins/Experimental/TechAudioTools/Source/TechAudioToolsMetaSoundEditor/Public/ViewModels/MetaSoundEditorViewModel.h b/Engine/Plugins/Experimental/TechAudioTools/Source/TechAudioToolsMetaSoundEditor/Public/ViewModels/MetaSoundEditorViewModel.h
--- a/Engine/Plugins/Experimental/TechAudioTools/Source/TechAudioToolsMetaSoundEditor/Public/ViewModels/MetaSoundEditorViewModel.h
+++ b/Engine/Plugins/Experimental/TechAudioTools/Source/TechAudioToolsMetaSoundEditor/Public/ViewModels/MetaSoundEditorViewModel.h
@@ -67,6 +67,8 @@ protected:
 	TObjectPtr<UMetaSoundEditorBuilderListener> BuilderListener;
 
 	virtual void CreateInputViewModel(const FMetasoundFrontendClassInput& InInput) override;
+    virtual UMetaSoundInputViewModel* CreateInputViewModelInstance() override;
+    virtual UMetaSoundOutputViewModel* CreateOutputViewModelInstance() override;
 
 	void BindDelegates();
 	void UnbindDelegates();
