MetaSoundInputEditorViewModel - Fix for SetInputDisplayName calling the incorrect function

diff --git a/Engine/Plugins/Experimental/TechAudioTools/Source/TechAudioToolsMetaSoundEditor/Private/ViewModels/MetaSoundEditorViewModel.cpp b/Engine/Plugins/Experimental/TechAudioTools/Source/TechAudioToolsMetaSoundEditor/Private/ViewModels/MetaSoundEditorViewModel.cpp
--- a/Engine/Plugins/Experimental/TechAudioTools/Source/TechAudioToolsMetaSoundEditor/Private/ViewModels/MetaSoundEditorViewModel.cpp
+++ b/Engine/Plugins/Experimental/TechAudioTools/Source/TechAudioToolsMetaSoundEditor/Private/ViewModels/MetaSoundEditorViewModel.cpp
@@ -335,7 +335,7 @@ void UMetaSoundInputEditorViewModel::SetInputDisplayName(const FText& InDisplayN
 		FMetaSoundFrontendDocumentBuilder& DocumentBuilder = Builder->GetBuilder();
 		if (DocumentBuilder.IsValid() && UE_MVVM_SET_PROPERTY_VALUE(InputDisplayName, InDisplayName))
 		{
-			DocumentBuilder.SetGraphOutputDisplayName(InputName, InDisplayName);
+			DocumentBuilder.SetGraphInputDisplayName(InputName, InDisplayName);
 		}
 	}
 }
