TechAudioTools - added float mapping class

#rnx

diff --git a/Engine/Plugins/Experimental/TechAudioTools/Source/TechAudioTools/Private/TechAudioToolsFloatMapping.cpp b/Engine/Plugins/Experimental/TechAudioTools/Source/TechAudioTools/Private/TechAudioToolsFloatMapping.cpp
new file mode 100644
--- /dev/null
+++ b/Engine/Plugins/Experimental/TechAudioTools/Source/TechAudioTools/Private/TechAudioToolsFloatMapping.cpp
@@ -0,0 +1,127 @@
+﻿// Copyright Epic Games, Inc. All Rights Reserved.
+
+#include "TechAudioToolsFloatMapping.h"
+
+#include UE_INLINE_GENERATED_CPP_BY_NAME(TechAudioToolsFloatMapping)
+
+FFloatInterval UTechAudioToolsFloatMapping::GetRange(const ETechAudioToolsMappingEndpoint Endpoint) const
+{
+	switch (MappingType)
+	{
+	case ETechAudioToolsFloatMappingType::Default:
+		return Range;
+
+	case ETechAudioToolsFloatMappingType::MapRange:
+		return Endpoint == ETechAudioToolsMappingEndpoint::Display ? DisplayRange : SourceRange;
+
+	case ETechAudioToolsFloatMappingType::Volume:
+		{
+			const bool bShowDecibels = Endpoint == ETechAudioToolsMappingEndpoint::Display ? DisplayVolumeUnits == ETechAudioToolsVolumeUnit::Decibels : SourceVolumeUnits == ETechAudioToolsVolumeUnit::Decibels;
+			return bShowDecibels ? DisplayRange_Decibels : DisplayRange_LinearGain;
+		}
+
+	case ETechAudioToolsFloatMappingType::Pitch:
+		{
+			const bool bShowSemitones = Endpoint == ETechAudioToolsMappingEndpoint::Display ? DisplayPitchUnits == ETechAudioToolsPitchUnit::Semitones : SourcePitchUnits == ETechAudioToolsPitchUnit::Semitones;
+			return bShowSemitones ? DisplayRange_Semitones : DisplayRange_FrequencyMultiplier;
+		}
+	}
+
+	ensureMsgf(false, TEXT("Unhandled MappingType"));
+	return FFloatInterval(0.f, 1.f);
+}
+
+ETechAudioToolsFloatUnit UTechAudioToolsFloatMapping::GetUnits(const ETechAudioToolsMappingEndpoint Endpoint) const
+{
+	using namespace TechAudioTools;
+
+	switch (MappingType)
+	{
+	case ETechAudioToolsFloatMappingType::Default:
+	case ETechAudioToolsFloatMappingType::MapRange:
+		return DisplayUnits;
+
+	case ETechAudioToolsFloatMappingType::Volume:
+		{
+			const bool bShowDecibels = Endpoint == ETechAudioToolsMappingEndpoint::Display ? DisplayVolumeUnits == ETechAudioToolsVolumeUnit::Decibels : SourceVolumeUnits == ETechAudioToolsVolumeUnit::Decibels;
+			return bShowDecibels ? ETechAudioToolsFloatUnit::Decibels : ETechAudioToolsFloatUnit::Multiplier;
+		}
+
+	case ETechAudioToolsFloatMappingType::Pitch:
+		{
+			const bool bShowSemitones = Endpoint == ETechAudioToolsMappingEndpoint::Display ? DisplayPitchUnits == ETechAudioToolsPitchUnit::Semitones : SourcePitchUnits == ETechAudioToolsPitchUnit::Semitones;
+			return bShowSemitones ? ETechAudioToolsFloatUnit::Semitones : ETechAudioToolsFloatUnit::Multiplier;
+		}
+	}
+
+	ensureMsgf(false, TEXT("Unhandled MappingType"));
+	return ETechAudioToolsFloatUnit::None;
+}
+
+float UTechAudioToolsFloatMapping::SourceToDisplay(const float InSource) const
+{
+	switch (MappingType)
+	{
+	case ETechAudioToolsFloatMappingType::Default:
+		return FMath::Clamp(InSource, Range.Min, Range.Max);
+
+	case ETechAudioToolsFloatMappingType::MapRange:
+		{
+			const float Normalized = FMath::GetRangePct(GetSourceMin(), GetSourceMax(), InSource);
+			return FMath::Lerp(GetDisplayMin(), GetDisplayMax(), Normalized);
+		}
+
+	case ETechAudioToolsFloatMappingType::Volume:
+		return TechAudioTools::ConvertUnit(SourceVolumeUnits, DisplayVolumeUnits, InSource);
+
+	case ETechAudioToolsFloatMappingType::Pitch:
+		return TechAudioTools::ConvertUnit(SourcePitchUnits, DisplayPitchUnits, InSource);
+	}
+
+	ensureMsgf(false, TEXT("Unhandled MappingType"));
+	return InSource;
+}
+
+float UTechAudioToolsFloatMapping::DisplayToSource(const float InDisplay) const
+{
+	switch (MappingType)
+	{
+	case ETechAudioToolsFloatMappingType::Default:
+		return FMath::Clamp(InDisplay, Range.Min, Range.Max);
+
+	case ETechAudioToolsFloatMappingType::MapRange:
+		{
+			const float Normalized = FMath::GetRangePct(GetDisplayMin(), GetDisplayMax(), InDisplay);
+			return FMath::Lerp(GetSourceMin(), GetSourceMax(), Normalized);
+		}
+
+	case ETechAudioToolsFloatMappingType::Volume:
+		return TechAudioTools::ConvertUnit(DisplayVolumeUnits, SourceVolumeUnits, InDisplay);
+
+	case ETechAudioToolsFloatMappingType::Pitch:
+		return TechAudioTools::ConvertUnit(DisplayPitchUnits, SourcePitchUnits, InDisplay);
+	}
+
+	ensureMsgf(false, TEXT("Unhandled MappingType"));
+	return InDisplay;
+}
+
+float UTechAudioToolsFloatMapping::SourceToNormalized(const float InSource) const
+{
+	return FMath::GetRangePct(GetSourceMin(), GetSourceMax(), InSource);
+}
+
+float UTechAudioToolsFloatMapping::NormalizedToSource(const float InNormalized) const
+{
+	return FMath::Lerp(GetSourceMin(), GetSourceMax(), InNormalized);
+}
+
+float UTechAudioToolsFloatMapping::DisplayToNormalized(const float InDisplay) const
+{
+	return FMath::GetRangePct(GetDisplayMin(), GetDisplayMax(), InDisplay);
+}
+
+float UTechAudioToolsFloatMapping::NormalizedToDisplay(const float InNormalized) const
+{
+	return FMath::Lerp(GetDisplayMin(), GetDisplayMax(), InNormalized);
+}
diff --git a/Engine/Plugins/Experimental/TechAudioTools/Source/TechAudioTools/Private/TechAudioToolsLog.cpp b/Engine/Plugins/Experimental/TechAudioTools/Source/TechAudioTools/Private/TechAudioToolsLog.cpp
new file mode 100644
--- /dev/null
+++ b/Engine/Plugins/Experimental/TechAudioTools/Source/TechAudioTools/Private/TechAudioToolsLog.cpp
@@ -0,0 +1,5 @@
+﻿// Copyright Epic Games, Inc. All Rights Reserved.
+
+#include "TechAudioToolsLog.h"
+
+DEFINE_LOG_CATEGORY(LogTechAudioTools);
diff --git a/Engine/Plugins/Experimental/TechAudioTools/Source/TechAudioTools/Private/TechAudioToolsModule.cpp b/Engine/Plugins/Experimental/TechAudioTools/Source/TechAudioTools/Private/TechAudioToolsModule.cpp
new file mode 100644
--- /dev/null
+++ b/Engine/Plugins/Experimental/TechAudioTools/Source/TechAudioTools/Private/TechAudioToolsModule.cpp
@@ -0,0 +1,5 @@
+﻿// Copyright Epic Games, Inc. All Rights Reserved.
+
+#include "Modules/ModuleManager.h"
+
+IMPLEMENT_MODULE(FDefaultModuleImpl, TechAudioTools)
diff --git a/Engine/Plugins/Experimental/TechAudioTools/Source/TechAudioTools/Private/Viewmodels/TechAudioToolsConversionFunctions.cpp b/Engine/Plugins/Experimental/TechAudioTools/Source/TechAudioTools/Private/Viewmodels/TechAudioToolsConversionFunctions.cpp
new file mode 100644
--- /dev/null
+++ b/Engine/Plugins/Experimental/TechAudioTools/Source/TechAudioTools/Private/Viewmodels/TechAudioToolsConversionFunctions.cpp
@@ -0,0 +1,69 @@
+// Copyright Epic Games, Inc. All Rights Reserved.
+
+#include "Viewmodels/TechAudioToolsConversionFunctions.h"
+
+#include "TechAudioToolsStyle.h"
+
+#include UE_INLINE_GENERATED_CPP_BY_NAME(TechAudioToolsConversionFunctions)
+
+FText UTechAudioToolsConversionFunctions::GetFloatValueFormatted(const float InValue, const ETechAudioToolsFloatUnit DisplayUnit)
+{
+	using namespace TechAudioTools;
+
+	FText FormattedFloat;
+	FText FloatUnitLabel;
+
+	switch (DisplayUnit)
+	{
+	case ETechAudioToolsFloatUnit::None:
+		FormattedFloat = FText::AsNumber(InValue, GetDefaultFloatFormat());
+		break;
+	case ETechAudioToolsFloatUnit::Decibels:
+		FormattedFloat = FText::AsNumber(InValue, GetDecibelsFloatFormat());
+		FloatUnitLabel = FText::FromString(TEXT(" dB"));
+		break;
+	case ETechAudioToolsFloatUnit::Hertz:
+		FormattedFloat = FText::AsNumber(InValue, GetFrequencyFloatFormat());
+		FloatUnitLabel = FText::FromString(TEXT(" Hz"));
+		break;
+	case ETechAudioToolsFloatUnit::Semitones:
+		FormattedFloat = FText::AsNumber(InValue, GetSemitonesFloatFormat());
+		break;
+	case ETechAudioToolsFloatUnit::TimeSeconds:
+		FormattedFloat = FText::AsNumber(InValue, GetSecondsFloatFormat());
+		FloatUnitLabel = FText::FromString(TEXT("s"));
+		break;
+	case ETechAudioToolsFloatUnit::TimeMilliseconds:
+		FormattedFloat = FText::AsNumber(InValue, GetMillisecondsFloatFormat());
+		FloatUnitLabel = FText::FromString(TEXT("ms"));
+		break;
+	case ETechAudioToolsFloatUnit::Degrees:
+		FormattedFloat = FText::AsNumber(InValue, GetDegreesFloatFormat());
+		FloatUnitLabel = FText::FromString(TEXT("°"));
+		break;
+	case ETechAudioToolsFloatUnit::Multiplier:
+		FormattedFloat = FText::AsNumber(InValue, GetMultiplierFloatFormat());
+		FloatUnitLabel = FText::FromString(TEXT("x"));
+		break;
+	case ETechAudioToolsFloatUnit::Pan:
+		static constexpr float KCenterEpsilon = 0.005f;
+		if (FMath::IsNearlyZero(InValue, KCenterEpsilon))
+		{
+			FormattedFloat  = FText::FromString(TEXT(""));
+			FloatUnitLabel  = FText::FromString(TEXT("C"));
+		}
+		else if (InValue < 0.f)
+		{
+			FormattedFloat  = FText::AsNumber(FMath::Abs(InValue), GetPanFloatFormat());
+			FloatUnitLabel  = FText::FromString(TEXT(" L"));
+		}
+		else
+		{
+			FormattedFloat  = FText::AsNumber(FMath::Abs(InValue), GetPanFloatFormat());
+			FloatUnitLabel  = FText::FromString(TEXT(" R"));
+		}
+		break;
+	}
+
+	return FText::Format(FText::FromString(TEXT("{0}{1}")), FormattedFloat, FloatUnitLabel);
+}
diff --git a/Engine/Plugins/Experimental/TechAudioTools/Source/TechAudioTools/Public/TechAudioToolsFloatMapping.h b/Engine/Plugins/Experimental/TechAudioTools/Source/TechAudioTools/Public/TechAudioToolsFloatMapping.h
new file mode 100644
--- /dev/null
+++ b/Engine/Plugins/Experimental/TechAudioTools/Source/TechAudioTools/Public/TechAudioToolsFloatMapping.h
@@ -0,0 +1,137 @@
+﻿// Copyright Epic Games, Inc. All Rights Reserved.
+
+#pragma once
+
+#include "DSP/Dsp.h"
+#include "TechAudioToolsTypes.h"
+
+#include "TechAudioToolsFloatMapping.generated.h"
+
+namespace TechAudioTools
+{
+	// Converts volume values between linear gain and decibels
+	FORCEINLINE float ConvertUnit(const ETechAudioToolsVolumeUnit From, const ETechAudioToolsVolumeUnit To, const float Value)
+	{
+		if (From == To)
+		{
+			return Value;
+		}
+
+		return To == ETechAudioToolsVolumeUnit::Decibels ? Audio::ConvertToDecibels(Value) : Audio::ConvertToLinear(Value);
+	}
+
+	// Converts pitch values between semitones and frequency multiplier
+	FORCEINLINE float ConvertUnit(const ETechAudioToolsPitchUnit From, const ETechAudioToolsPitchUnit To, const float Value)
+	{
+		if (From == To)
+		{
+			return Value;
+		}
+
+		return To == ETechAudioToolsPitchUnit::Semitones ? Audio::GetSemitones(Value) : Audio::GetFrequencyMultiplier(Value);
+	}
+} // namespace TechAudioTools
+
+/**
+ * Base class for any float value mapping that translates between two domains:
+ * - Source domain: range used by an internal system (e.g. a MetaSound that uses linear gain or frequency multiplier).
+ * - Display domain: how that value is shown or edited in the UI (e.g. sliders or knobs that use decibels or semitones).
+ */
+UCLASS(MinimalAPI, BlueprintType, EditInlineNew, DefaultToInstanced, CollapseCategories, DisplayName = "TechAudioTools Float Mapping")
+class UTechAudioToolsFloatMapping : public UObject
+{
+	GENERATED_BODY()
+
+public:
+	// Selects the mapping mode.
+	UPROPERTY(EditAnywhere, Category = "Range")
+	ETechAudioToolsFloatMappingType MappingType = ETechAudioToolsFloatMappingType::Default;
+
+	//~ Default
+
+	// Default min/max values.
+	UPROPERTY(EditAnywhere, Category = "Range", meta = (EditCondition = "MappingType == ETechAudioToolsFloatMappingType::Default", EditConditionHides))
+	FFloatInterval Range { 0.f, 1.f };
+
+	//~ Linear Remap
+
+	// Range used by an internal system.
+	UPROPERTY(EditAnywhere, Category = "Range", meta = (EditCondition = "MappingType == ETechAudioToolsFloatMappingType::MapRange", EditConditionHides))
+	FFloatInterval SourceRange { 0.f, 1.f };
+
+	// Range displayed in a user interface.
+	UPROPERTY(EditAnywhere, Category = "Range", meta = (EditCondition = "MappingType == ETechAudioToolsFloatMappingType::MapRange", EditConditionHides))
+	FFloatInterval DisplayRange { 0.f, 1.f };
+
+	// Units used to label the display range.
+	UPROPERTY(BlueprintReadWrite, EditAnywhere, Category = "Map Range", meta = (EditCondition = "MappingType == ETechAudioToolsFloatMappingType::Default || MappingType == ETechAudioToolsFloatMappingType::MapRange", EditConditionHides))
+	ETechAudioToolsFloatUnit DisplayUnits = ETechAudioToolsFloatUnit::None;
+
+	//~ Volume
+
+	// Units of the internal volume value.
+	UPROPERTY(EditAnywhere, Category = "Volume", meta = (EditCondition = "MappingType == ETechAudioToolsFloatMappingType::Volume", EditConditionHides))
+	ETechAudioToolsVolumeUnit SourceVolumeUnits = ETechAudioToolsVolumeUnit::LinearGain;
+
+	// Units to show volume in a user interface.
+	UPROPERTY(EditAnywhere, Category = "Volume", meta = (EditCondition = "MappingType == ETechAudioToolsFloatMappingType::Volume", EditConditionHides))
+	ETechAudioToolsVolumeUnit DisplayVolumeUnits = ETechAudioToolsVolumeUnit::Decibels;
+
+	// Range in decibels to display in a user interface.
+	UPROPERTY(EditAnywhere, Category = "Volume", DisplayName = "Display Decibel Range", meta = (EditCondition = "MappingType == ETechAudioToolsFloatMappingType::Volume && DisplayVolumeUnits == ETechAudioToolsVolumeUnit::Decibels", EditConditionHides, ClampMin = "-60", ClampMax = "24"))
+	FFloatInterval DisplayRange_Decibels { -60.f, 6.f };
+
+	// Range in linear gain to display in a user interface.
+	UPROPERTY(EditAnywhere, Category = "Volume", DisplayName = "Display Linear Gain Range", meta = (EditCondition = "MappingType == ETechAudioToolsFloatMappingType::Volume && DisplayVolumeUnits == ETechAudioToolsVolumeUnit::LinearGain", EditConditionHides, ClampMin = "0.0", ClampMax = "8.0"))
+	FFloatInterval DisplayRange_LinearGain { 0.001f, 4.f };
+
+	//~ Pitch
+
+	// Units of the internal pitch value.
+	UPROPERTY(EditAnywhere, Category = "Pitch", meta = (EditCondition = "MappingType == ETechAudioToolsFloatMappingType::Pitch", EditConditionHides))
+	ETechAudioToolsPitchUnit SourcePitchUnits = ETechAudioToolsPitchUnit::Semitones;
+
+	// Units to show pitch in a user interface.
+	UPROPERTY(EditAnywhere, Category = "Pitch", meta = (EditCondition = "MappingType == ETechAudioToolsFloatMappingType::Pitch", EditConditionHides))
+	ETechAudioToolsPitchUnit DisplayPitchUnits = ETechAudioToolsPitchUnit::Semitones;
+
+	// Range in semitones to display in a user interface.
+	UPROPERTY(EditAnywhere, Category = "Pitch", DisplayName = "Display Semitone Range", meta = (EditCondition = "MappingType == ETechAudioToolsFloatMappingType::Pitch && DisplayPitchUnits == ETechAudioToolsPitchUnit::Semitones", EditConditionHides, ClampMin = "-120", ClampMax = "120"))
+	FFloatInterval DisplayRange_Semitones { -12.f, 12.f };
+
+	// Range in frequency multiplier to display in a user interface.
+	UPROPERTY(EditAnywhere, Category = "Pitch", DisplayName = "Display Frequency Multiplier Range", meta = (EditCondition = "MappingType == ETechAudioToolsFloatMappingType::Pitch && DisplayPitchUnits == ETechAudioToolsPitchUnit::FrequencyMultiplier", EditConditionHides, ClampMin = "0.125", ClampMax = "8"))
+	FFloatInterval DisplayRange_FrequencyMultiplier { 0.5f, 2.f };
+
+	// Returns the units used by the Source or Display endpoints.
+	UFUNCTION(BlueprintPure, Category = "Range")
+	TECHAUDIOTOOLS_API ETechAudioToolsFloatUnit GetUnits(ETechAudioToolsMappingEndpoint Endpoint) const;
+
+	// Returns the range minimum of the source range.
+	UFUNCTION(BlueprintPure, Category = "Range")
+	TECHAUDIOTOOLS_API float GetSourceMin() const { return GetRange(ETechAudioToolsMappingEndpoint::Source).Min; }
+
+	// Returns the range maximum of the source range.
+	UFUNCTION(BlueprintPure, Category = "Range")
+	TECHAUDIOTOOLS_API float GetSourceMax() const { return GetRange(ETechAudioToolsMappingEndpoint::Source).Max; }
+
+	// Returns the range minimum of the display range.
+	UFUNCTION(BlueprintPure, Category = "Range")
+	TECHAUDIOTOOLS_API float GetDisplayMin() const { return GetRange(ETechAudioToolsMappingEndpoint::Display).Min; }
+
+	// Returns the range maximum of the display range.
+	UFUNCTION(BlueprintPure, Category = "Range")
+	TECHAUDIOTOOLS_API float GetDisplayMax() const { return GetRange(ETechAudioToolsMappingEndpoint::Display).Max; }
+
+	TECHAUDIOTOOLS_API float SourceToDisplay(float InSource) const;
+	TECHAUDIOTOOLS_API float DisplayToSource(float InDisplay) const;
+
+	TECHAUDIOTOOLS_API float NormalizedToSource(float InNormalized) const;
+	TECHAUDIOTOOLS_API float SourceToNormalized(float InSource) const;
+
+	TECHAUDIOTOOLS_API float NormalizedToDisplay(float InNormalized) const;
+	TECHAUDIOTOOLS_API float DisplayToNormalized(float InDisplay) const;
+
+private:
+	FFloatInterval GetRange(const ETechAudioToolsMappingEndpoint Endpoint) const;
+};
diff --git a/Engine/Plugins/Experimental/TechAudioTools/Source/TechAudioTools/Public/TechAudioToolsLog.h b/Engine/Plugins/Experimental/TechAudioTools/Source/TechAudioTools/Public/TechAudioToolsLog.h
new file mode 100644
--- /dev/null
+++ b/Engine/Plugins/Experimental/TechAudioTools/Source/TechAudioTools/Public/TechAudioToolsLog.h
@@ -0,0 +1,7 @@
+﻿// Copyright Epic Games, Inc. All Rights Reserved.
+
+#pragma once
+
+#include "Logging/LogMacros.h"
+
+TECHAUDIOTOOLS_API DECLARE_LOG_CATEGORY_EXTERN(LogTechAudioTools, Log, All);
diff --git a/Engine/Plugins/Experimental/TechAudioTools/Source/TechAudioTools/Public/TechAudioToolsStyle.h b/Engine/Plugins/Experimental/TechAudioTools/Source/TechAudioTools/Public/TechAudioToolsStyle.h
new file mode 100644
--- /dev/null
+++ b/Engine/Plugins/Experimental/TechAudioTools/Source/TechAudioTools/Public/TechAudioToolsStyle.h
@@ -0,0 +1,116 @@
+// Copyright Epic Games, Inc. All Rights Reserved.
+
+#pragma once
+
+#include "Internationalization/Text.h"
+
+namespace TechAudioTools
+{
+	const FNumberFormattingOptions* GetDefaultFloatFormat()
+	{
+		static FNumberFormattingOptions DefaultFormat;
+		DefaultFormat.RoundingMode = HalfToEven;
+		DefaultFormat.AlwaysSign = false;
+		DefaultFormat.UseGrouping = true;
+		DefaultFormat.MinimumIntegralDigits = 1;
+		DefaultFormat.MinimumFractionalDigits = 0;
+		DefaultFormat.MaximumFractionalDigits = 2;
+		return &DefaultFormat;
+	}
+
+	const FNumberFormattingOptions* GetMultiplierFloatFormat()
+	{
+		static FNumberFormattingOptions MultiplierFormat;
+		MultiplierFormat.RoundingMode = HalfToEven;
+		MultiplierFormat.AlwaysSign = false;
+		MultiplierFormat.UseGrouping = true;
+		MultiplierFormat.MinimumIntegralDigits = 1;
+		MultiplierFormat.MinimumFractionalDigits = 0;
+		MultiplierFormat.MaximumFractionalDigits = 2;
+		return &MultiplierFormat;
+	}
+
+	const FNumberFormattingOptions* GetDecibelsFloatFormat()
+	{
+		static FNumberFormattingOptions DecibelsFormat;
+		DecibelsFormat.RoundingMode = HalfToEven;
+		DecibelsFormat.AlwaysSign = true;
+		DecibelsFormat.UseGrouping = true;
+		DecibelsFormat.MinimumIntegralDigits = 1;
+		DecibelsFormat.MinimumFractionalDigits = 0;
+		DecibelsFormat.MaximumFractionalDigits = 2;
+		return &DecibelsFormat;
+	}
+
+	const FNumberFormattingOptions* GetFrequencyFloatFormat()
+	{
+		static FNumberFormattingOptions FrequencyFormat;
+		FrequencyFormat.RoundingMode = HalfToEven;
+		FrequencyFormat.AlwaysSign = false;
+		FrequencyFormat.UseGrouping = true;
+		FrequencyFormat.MinimumIntegralDigits = 1;
+		FrequencyFormat.MinimumFractionalDigits = 0;
+		FrequencyFormat.MaximumFractionalDigits = 1;
+		return &FrequencyFormat;
+	}
+
+	const FNumberFormattingOptions* GetSemitonesFloatFormat()
+	{
+		static FNumberFormattingOptions SemitonesFormat;
+		SemitonesFormat.RoundingMode = HalfToEven;
+		SemitonesFormat.AlwaysSign = true;
+		SemitonesFormat.UseGrouping = true;
+		SemitonesFormat.MinimumIntegralDigits = 1;
+		SemitonesFormat.MinimumFractionalDigits = 0;
+		SemitonesFormat.MaximumFractionalDigits = 2;
+		return &SemitonesFormat;
+	}
+
+	const FNumberFormattingOptions* GetSecondsFloatFormat()
+	{
+		static FNumberFormattingOptions SecondsFormat;
+		SecondsFormat.RoundingMode = HalfToEven;
+		SecondsFormat.AlwaysSign = false;
+		SecondsFormat.UseGrouping = true;
+		SecondsFormat.MinimumIntegralDigits = 1;
+		SecondsFormat.MinimumFractionalDigits = 0;
+		SecondsFormat.MaximumFractionalDigits = 2;
+		return &SecondsFormat;
+	}
+
+	const FNumberFormattingOptions* GetMillisecondsFloatFormat()
+	{
+		static FNumberFormattingOptions MillisecondsFormat;
+		MillisecondsFormat.RoundingMode = HalfToEven;
+		MillisecondsFormat.AlwaysSign = false;
+		MillisecondsFormat.UseGrouping = true;
+		MillisecondsFormat.MinimumIntegralDigits = 1;
+		MillisecondsFormat.MinimumFractionalDigits = 0;
+		MillisecondsFormat.MaximumFractionalDigits = 0;
+		return &MillisecondsFormat;
+	}
+
+	const FNumberFormattingOptions* GetDegreesFloatFormat()
+	{
+		static FNumberFormattingOptions DegreesFormat;
+		DegreesFormat.RoundingMode = HalfToEven;
+		DegreesFormat.AlwaysSign = false;
+		DegreesFormat.UseGrouping = true;
+		DegreesFormat.MinimumIntegralDigits = 1;
+		DegreesFormat.MinimumFractionalDigits = 0;
+		DegreesFormat.MaximumFractionalDigits = 2;
+		return &DegreesFormat;
+	}
+
+	const FNumberFormattingOptions* GetPanFloatFormat()
+	{
+		static FNumberFormattingOptions PanFormat;
+		PanFormat.RoundingMode = HalfToEven;
+		PanFormat.AlwaysSign = false;
+		PanFormat.UseGrouping = true;
+		PanFormat.MinimumIntegralDigits = 1;
+		PanFormat.MinimumFractionalDigits = 0;
+		PanFormat.MaximumFractionalDigits = 1;
+		return &PanFormat;
+	}
+}
diff --git a/Engine/Plugins/Experimental/TechAudioTools/Source/TechAudioTools/Public/TechAudioToolsTypes.h b/Engine/Plugins/Experimental/TechAudioTools/Source/TechAudioTools/Public/TechAudioToolsTypes.h
new file mode 100644
--- /dev/null
+++ b/Engine/Plugins/Experimental/TechAudioTools/Source/TechAudioTools/Public/TechAudioToolsTypes.h
@@ -0,0 +1,54 @@
+// Copyright Epic Games, Inc. All Rights Reserved.
+
+#pragma once
+
+#include "TechAudioToolsTypes.generated.h"
+
+// Identifies which side of a mapping a value belongs to.
+UENUM(BlueprintType)
+enum class ETechAudioToolsMappingEndpoint : uint8
+{
+	Source				UMETA(DisplayName = "Source",				ToolTip = "Refers to values intended to be used by an internal system (e.g. a MetaSound"),
+	Display				UMETA(DisplayName = "Display",				ToolTip = "Refers to values intended to be displayed via a user interface (e.g. a MetaSound preset widget"),
+};
+
+// Float unit types used primarily for formatting values.
+UENUM(BlueprintType)
+enum class ETechAudioToolsFloatUnit : uint8
+{
+	None				UMETA(DisplayName = "None"),
+	Multiplier  		UMETA(DisplayName = "Multiplier",   		ToolTip = "Linear multiplier (1.0 = unity)"),
+	Decibels			UMETA(DisplayName = "Decibels", 			ToolTip = "Amplitude in dB (0 dB = unity)"),
+	Hertz   			UMETA(DisplayName = "Hertz",				ToolTip = "Frequency in Hertz (e.g. 20-20,000)"),
+	Semitones   		UMETA(DisplayName = "Semitones",			ToolTip = "Pitch offset in semitones (-12 = one octave down)"),
+	TimeSeconds 		UMETA(DisplayName = "Seconds",  			ToolTip = "Time in seconds"),
+	TimeMilliseconds	UMETA(DisplayName = "Milliseconds", 		ToolTip = "Time in ms"),
+	Degrees 			UMETA(DisplayName = "Degrees",  			ToolTip = "Angle in degrees (0-360)"),
+	Pan 				UMETA(DisplayName = "Pan",  				ToolTip = "Pan value: negative = L, 0 = C, positive = R"),
+};
+
+// The mapping logic applied between source and display values.
+UENUM(BlueprintType)
+enum class ETechAudioToolsFloatMappingType : uint8
+{
+	Default 			UMETA(DisplayName = "Default",				ToolTip = "No mapping is applied. Display range is the same as the source range."),
+	MapRange			UMETA(DisplayName = "Map Range",			ToolTip = "Linear remap between the source and display ranges. Does not convert units (e.g. seconds will not be converted to milliseconds)."),
+	Volume				UMETA(DisplayName = "Volume",				ToolTip = "Remaps a given volume range and converts between Linear Gain and Decibels if necessary."),
+	Pitch				UMETA(DisplayName = "Pitch",				ToolTip = "Remaps a given pitch range and converts between Semitones and Frequency Multiplier if necessary."),
+};
+
+// Defines the units used when interpreting a float as volume.
+UENUM()
+enum class ETechAudioToolsVolumeUnit : uint8
+{
+	Decibels			UMETA(DisplayName = "Decibels", 			ToolTip = "Decibels (logarithmic)"),
+	LinearGain  		UMETA(DisplayName = "Linear Gain",   		ToolTip = "Linear gain multiplier"),
+};
+
+// Defines the units used when interpreting a float as pitch.
+UENUM()
+enum class ETechAudioToolsPitchUnit : uint8
+{
+	Semitones   		UMETA(DisplayName = "Semitones",			ToolTip = "Pitch offset in semitones"),
+	FrequencyMultiplier UMETA(DisplayName = "Frequency Multiplier",	ToolTip = "Frequency multiplier (2.0 = one octave up)"),
+};
diff --git a/Engine/Plugins/Experimental/TechAudioTools/Source/TechAudioTools/Public/Viewmodels/TechAudioToolsConversionFunctions.h b/Engine/Plugins/Experimental/TechAudioTools/Source/TechAudioTools/Public/Viewmodels/TechAudioToolsConversionFunctions.h
new file mode 100644
--- /dev/null
+++ b/Engine/Plugins/Experimental/TechAudioTools/Source/TechAudioTools/Public/Viewmodels/TechAudioToolsConversionFunctions.h
@@ -0,0 +1,19 @@
+// Copyright Epic Games, Inc. All Rights Reserved.
+
+#pragma once
+
+#include "TechAudioToolsTypes.h"
+#include "Kismet/BlueprintFunctionLibrary.h"
+
+#include "TechAudioToolsConversionFunctions.generated.h"
+
+UCLASS()
+class TECHAUDIOTOOLS_API UTechAudioToolsConversionFunctions : public UBlueprintFunctionLibrary
+{
+	GENERATED_BODY()
+
+public:
+	// Provides a float value as formatted text using the given display unit.
+	UFUNCTION(BlueprintCallable, BlueprintPure, Category = "TechAudioTools", DisplayName = "Get Float Value Formatted (TechAudioTools)")
+	static FText GetFloatValueFormatted(float InValue, ETechAudioToolsFloatUnit DisplayUnit);
+};
diff --git a/Engine/Plugins/Experimental/TechAudioTools/Source/TechAudioTools/TechAudioTools.Build.cs b/Engine/Plugins/Experimental/TechAudioTools/Source/TechAudioTools/TechAudioTools.Build.cs
new file mode 100644
--- /dev/null
+++ b/Engine/Plugins/Experimental/TechAudioTools/Source/TechAudioTools/TechAudioTools.Build.cs
@@ -0,0 +1,26 @@
+﻿// Copyright Epic Games, Inc. All Rights Reserved.
+
+using UnrealBuildTool;
+
+public class TechAudioTools : ModuleRules
+{
+    public TechAudioTools(ReadOnlyTargetRules Target) : base(Target)
+    {
+        PCHUsage = ModuleRules.PCHUsageMode.UseExplicitOrSharedPCHs;
+
+        PublicDependencyModuleNames.AddRange(
+            new string[]
+            {
+                "Core",
+            }
+        );
+
+        PrivateDependencyModuleNames.AddRange(
+            new string[]
+            {
+                "CoreUObject",
+                "Engine",
+            }
+        );
+    }
+}
diff --git a/Engine/Plugins/Experimental/TechAudioTools/TechAudioTools.uplugin b/Engine/Plugins/Experimental/TechAudioTools/TechAudioTools.uplugin
--- a/Engine/Plugins/Experimental/TechAudioTools/TechAudioTools.uplugin
+++ b/Engine/Plugins/Experimental/TechAudioTools/TechAudioTools.uplugin
@@ -24,6 +24,11 @@
 			"Name": "TechAudioToolsMetaSoundEditor",
 			"Type": "Editor",
 			"LoadingPhase": "Default"
+		},
+		{
+			"Name": "TechAudioTools",
+			"Type": "Runtime",
+			"LoadingPhase": "Default"
 		}
 	],
 	"Plugins": [
