Set TechAudioToolsConversionFunctions to MinimalAPI

#rnx

diff --git a/Engine/Plugins/Experimental/TechAudioTools/Source/TechAudioTools/Public/Viewmodels/TechAudioToolsConversionFunctions.h b/Engine/Plugins/Experimental/TechAudioTools/Source/TechAudioTools/Public/Viewmodels/TechAudioToolsConversionFunctions.h
--- a/Engine/Plugins/Experimental/TechAudioTools/Source/TechAudioTools/Public/Viewmodels/TechAudioToolsConversionFunctions.h
+++ b/Engine/Plugins/Experimental/TechAudioTools/Source/TechAudioTools/Public/Viewmodels/TechAudioToolsConversionFunctions.h
@@ -7,13 +7,13 @@
 
 #include "TechAudioToolsConversionFunctions.generated.h"
 
-UCLASS()
-class TECHAUDIOTOOLS_API UTechAudioToolsConversionFunctions : public UBlueprintFunctionLibrary
+UCLASS(MinimalAPI)
+class UTechAudioToolsConversionFunctions : public UBlueprintFunctionLibrary
 {
 	GENERATED_BODY()
 
 public:
 	// Provides a float value as formatted text using the given display unit.
 	UFUNCTION(BlueprintCallable, BlueprintPure, Category = "TechAudioTools", DisplayName = "Get Float Value Formatted (TechAudioTools)")
-	static FText GetFloatValueFormatted(float InValue, ETechAudioToolsFloatUnit DisplayUnit);
+	static TECHAUDIOTOOLS_API FText GetFloatValueFormatted(float InValue, ETechAudioToolsFloatUnit DisplayUnit);
 };
