MetaSoundViewModel updates to use new BuilderListener delegates

- Implemented new delegates for various MetaSound document and member property change events
- Refactored previous delegate implementation for greatly improved optimization and cleaner code
- Fleshed out viewmodel getter/setter functions for native access
- Removed LogTechAudioToolsMetaSound log category (replaced with LogTechAudioTools log category)

#rnx

diff --git a/Engine/Plugins/Experimental/TechAudioTools/Source/TechAudioToolsMetaSound/Private/ViewModels/MetaSoundViewModel.cpp b/Engine/Plugins/Experimental/TechAudioTools/Source/TechAudioToolsMetaSound/Private/ViewModels/MetaSoundViewModel.cpp
--- a/Engine/Plugins/Experimental/TechAudioTools/Source/TechAudioToolsMetaSound/Private/ViewModels/MetaSoundViewModel.cpp
+++ b/Engine/Plugins/Experimental/TechAudioTools/Source/TechAudioToolsMetaSound/Private/ViewModels/MetaSoundViewModel.cpp
@@ -6,37 +6,11 @@
 #include "MetasoundDocumentBuilderRegistry.h"
 #include "MetasoundFrontendPages.h"
 #include "MetasoundOutputSubsystem.h"
-#include "MetasoundSettings.h"
+#include "TechAudioToolsLog.h"
+#include "ViewModels/MetaSoundViewModelConversionFunctions.h"
 
 #include UE_INLINE_GENERATED_CPP_BY_NAME(MetaSoundViewModel)
 
-DEFINE_LOG_CATEGORY(LogTechAudioToolsMetaSound);
-
-namespace TechAudioTools::MetaSoundViewModel
-{
-	FName GetAdjustedDataType(const FName& CurrentDataType, const bool bIsArray)
-	{
-		const FString DataTypeString = CurrentDataType.ToString();
-		if (bIsArray)
-		{
-			// Append the suffix if it isn't already there
-			return DataTypeString.EndsWith(TEXT(":Array")) ? CurrentDataType : FName(DataTypeString + TEXT(":Array"));
-		}
-		else
-		{
-			// Split on ':' and keep the base data type
-			FString BaseDataType;
-			DataTypeString.Split(TEXT(":"), &BaseDataType, nullptr);
-			return FName(BaseDataType);
-		}
-	}
-} // TechAudioTools::MetaSoundViewModel
-
-FText UMetaSoundViewModel::GetBuilderNameAsText() const
-{
-	return FText::FromString(GetNameSafe(Builder));
-}
-
 TArray<UMetaSoundInputViewModel*> UMetaSoundViewModel::GetInputViewModels() const
 {
 	TArray<TObjectPtr<UMetaSoundInputViewModel>> InputViewModelsArray;
@@ -69,7 +43,7 @@ void UMetaSoundViewModel::Initialize(UMetaSoundBuilderBase* InBuilder)
 
 	if (!InBuilder)
 	{
-		UE_LOGFMT(LogTechAudioToolsMetaSound, Log, "Unable to initialize MetaSoundViewModel. Builder was null.");
+		UE_LOGFMT(LogTechAudioTools, Warning, "Unable to initialize MetaSoundViewModel. Builder was null.");
 		return;
 	}
 
@@ -78,7 +52,7 @@ void UMetaSoundViewModel::Initialize(UMetaSoundBuilderBase* InBuilder)
 	const FMetaSoundFrontendDocumentBuilder& DocumentBuilder = InBuilder->GetConstBuilder();
 	const FMetasoundFrontendDocument& FrontendDocument = DocumentBuilder.GetConstDocumentChecked();
 
-	InitializeProperties(FrontendDocument);
+	UE_MVVM_SET_PROPERTY_VALUE(bIsPreset, FrontendDocument.RootGraph.PresetOptions.bIsPreset);
 
 	CreateMemberViewModels();
 	SetIsInitialized(true);
@@ -87,7 +61,7 @@ void UMetaSoundViewModel::Initialize(UMetaSoundBuilderBase* InBuilder)
 void UMetaSoundViewModel::Reset()
 {
 	SetIsInitialized(false);
-	ResetProperties();
+	UE_MVVM_SET_PROPERTY_VALUE(bIsPreset, false);
 	Builder = nullptr;
 
 	InputViewModels.Empty();
@@ -97,195 +71,464 @@ void UMetaSoundViewModel::Reset()
 	UE_MVVM_BROADCAST_FIELD_VALUE_CHANGED(GetOutputViewModels);
 }
 
-void UMetaSoundViewModel::InitializeProperties(const FMetasoundFrontendDocument& FrontendDocument)
+void UMetaSoundViewModel::SetIsInitialized(const bool bInIsInitialized)
 {
-	UE_MVVM_SET_PROPERTY_VALUE(bIsPreset, FrontendDocument.RootGraph.PresetOptions.bIsPreset);
+	UE_MVVM_SET_PROPERTY_VALUE(bIsInitialized, bInIsInitialized);
 }
 
-void UMetaSoundViewModel::ResetProperties()
+void UMetaSoundViewModel::SetInputName(const FName& OldName, const FName& NewName) const
 {
-	UE_MVVM_SET_PROPERTY_VALUE(bIsPreset, false);
+	if (OldName == NewName)
+	{
+		return;
+	}
+
+	if (InputViewModels.Contains(NewName))
+	{
+		UE_LOGFMT(LogTechAudioTools, Warning, "Cannot set input name '{NewName}'. An input already exists with that name.", NewName);
+		return;
+	}
+
+	if (UMetaSoundViewModelConversionFunctions::IsInterfaceMember(NewName))
+	{
+		UE_LOGFMT(LogTechAudioTools, Warning, "Cannot set input name '{NewName}'. Name already belongs to an interface.", NewName);
+		return;
+	}
+
+	if (Builder)
+	{
+		EMetaSoundBuilderResult Result;
+		Builder->SetGraphInputName(OldName, NewName, Result);
+	}
 }
 
-void UMetaSoundViewModel::CreateMemberViewModels()
+void UMetaSoundViewModel::SetInputDataType(const FName& InputName, const FName& DataType) const
 {
-	const FMetaSoundFrontendDocumentBuilder& DocumentBuilder = Builder->GetConstBuilder();
-	const FMetasoundFrontendDocument& FrontendDocument = DocumentBuilder.GetConstDocumentChecked();
+	if (!Metasound::Frontend::IDataTypeRegistry::Get().IsRegistered(DataType))
+	{
+		UE_LOGFMT(LogTechAudioTools, Warning, "Failed to set data type for {InputName}. '{DataType}' is not a registered data type.", InputName, DataType);
+		return;
+	}
 
-	for (const FMetasoundFrontendClassInput& Input : FrontendDocument.RootGraph.GetDefaultInterface().Inputs)
+	if (Builder)
 	{
-		CreateInputViewModel(Input);
+		FMetaSoundFrontendDocumentBuilder& DocumentBuilder = Builder->GetBuilder();
+		DocumentBuilder.SetGraphInputDataType(InputName, DataType);
 	}
-	UE_MVVM_BROADCAST_FIELD_VALUE_CHANGED(GetInputViewModels);
+}
 
-	for (const FMetasoundFrontendClassOutput& Output : FrontendDocument.RootGraph.GetDefaultInterface().Outputs)
+void UMetaSoundViewModel::SetInputDefaultLiteral(const FName& InputName, const FMetasoundFrontendLiteral& DefaultLiteral) const
+{
+	if (Builder)
 	{
-		CreateOutputViewModel(Output);
+		EMetaSoundBuilderResult Result;
+		Builder->SetGraphInputDefault(InputName, DefaultLiteral, Result);
 	}
-	UE_MVVM_BROADCAST_FIELD_VALUE_CHANGED(GetOutputViewModels);
 }
 
-void UMetaSoundViewModel::CreateInputViewModel(const FMetasoundFrontendClassInput& InInput)
+void UMetaSoundViewModel::SetInputOverridesDefault(const FName& InputName, const bool bOverridesDefault) const
 {
-	UMetaSoundInputViewModel* InputViewModel = CreateInputViewModelInstance();
-	check(InputViewModel);
+	if (Builder)
+	{
+		const bool bInheritsDefaultValue = !bOverridesDefault;
+		FMetaSoundFrontendDocumentBuilder& DocumentBuilder = Builder->GetBuilder();
+		DocumentBuilder.SetGraphInputInheritsDefault(InputName, bInheritsDefaultValue);
+	}
+}
 
-	InputViewModel->InitializeInput(Builder.Get(), InInput);
-	InputViewModels.Emplace(InInput.Name, InputViewModel);
+void UMetaSoundViewModel::SetInputIsConstructorPin(const FName& InputName, const bool bIsConstructorPin) const
+{
+	if (const UMetaSoundInputViewModel* InputViewModel = *InputViewModels.Find(InputName))
+	{
+		const FName DataType = InputViewModel->GetDataType();
+		if (!UMetaSoundViewModelConversionFunctions::IsConstructorType(DataType))
+		{
+			UE_LOGFMT(LogTechAudioTools, Warning, "Failed to set bIsConstructorPin on {InputName}. Data type '{DataType}' cannot be a constructor pin.", InputName, DataType);
+			return;
+		}
+	}
 
-	// Prevent spamming field notify broadcasts while we're still initializing.
-	if (bIsInitialized)
+	if (Builder)
 	{
-		UE_MVVM_BROADCAST_FIELD_VALUE_CHANGED(GetInputViewModels);
+		FMetaSoundFrontendDocumentBuilder& DocumentBuilder = Builder->GetBuilder();
+		const EMetasoundFrontendVertexAccessType AccessType = bIsConstructorPin ? EMetasoundFrontendVertexAccessType::Value : EMetasoundFrontendVertexAccessType::Reference;
+		DocumentBuilder.SetGraphInputAccessType(InputName, AccessType);
 	}
 }
 
-UMetaSoundInputViewModel* UMetaSoundViewModel::CreateInputViewModelInstance()
+void UMetaSoundViewModel::SetOutputName(const FName& OldName, const FName& NewName) const
 {
-	return NewObject<UMetaSoundInputViewModel>(this, UMetaSoundInputViewModel::StaticClass());
+	if (OldName == NewName)
+	{
+		return;
+	}
+
+	if (OutputViewModels.Contains(NewName))
+	{
+		UE_LOGFMT(LogTechAudioTools, Warning, "Cannot set output name '{NewName}'. An output already exists with that name.", NewName);
+		return;
+	}
+
+	if (UMetaSoundViewModelConversionFunctions::IsInterfaceMember(NewName))
+	{
+		UE_LOGFMT(LogTechAudioTools, Warning, "Cannot set output name '{NewName}'. Name already belongs to an interface.", NewName);
+		return;
+	}
+
+	if (Builder)
+	{
+		EMetaSoundBuilderResult Result;
+		Builder->SetGraphOutputName(OldName, NewName, Result);
+	}
 }
 
-void UMetaSoundViewModel::CreateOutputViewModel(const FMetasoundFrontendClassOutput& InOutput)
+void UMetaSoundViewModel::SetOutputDataType(const FName& OutputName, const FName& DataType) const
 {
-	UMetaSoundOutputViewModel* OutputViewModel = CreateOutputViewModelInstance();
-	check(OutputViewModel);
+	if (!Metasound::Frontend::IDataTypeRegistry::Get().IsRegistered(DataType))
+	{
+		UE_LOGFMT(LogTechAudioTools, Warning, "Failed to set data type for {OutputName}. Data type '{DataType}' is not a registered data type.", OutputName, DataType);
+		return;
+	}
 
-	OutputViewModel->InitializeOutput(Builder.Get(), InOutput);
-	OutputViewModels.Emplace(InOutput.Name, OutputViewModel);
+	if (Builder)
+	{
+		FMetaSoundFrontendDocumentBuilder& DocumentBuilder = Builder->GetBuilder();
+		DocumentBuilder.SetGraphOutputDataType(OutputName, DataType);
+	}
+}
 
-	if (bIsInitialized)
+void UMetaSoundViewModel::SetOutputIsConstructorPin(const FName& OutputName, const bool bIsConstructorPin) const
+{
+	if (const UMetaSoundOutputViewModel* OutputViewModel = *OutputViewModels.Find(OutputName))
 	{
-		UE_MVVM_BROADCAST_FIELD_VALUE_CHANGED(GetOutputViewModels);
+		const FName DataType = OutputViewModel->GetDataType();
+		if (!UMetaSoundViewModelConversionFunctions::IsConstructorType(DataType))
+		{
+			UE_LOGFMT(LogTechAudioTools, Warning, "Failed to set bIsConstructorPin on {OutputName}. Data type '{DataType}' cannot be a constructor pin.", OutputName, DataType);
+			return;
+		}
+	}
+
+	if (Builder)
+	{
+		FMetaSoundFrontendDocumentBuilder& DocumentBuilder = Builder->GetBuilder();
+		const EMetasoundFrontendVertexAccessType AccessType = bIsConstructorPin ? EMetasoundFrontendVertexAccessType::Value : EMetasoundFrontendVertexAccessType::Reference;
+		DocumentBuilder.SetGraphOutputAccessType(OutputName, AccessType);
 	}
 }
 
-UMetaSoundOutputViewModel* UMetaSoundViewModel::CreateOutputViewModelInstance()
+void UMetaSoundViewModel::OnInputAdded(const FName VertexName, FName DataType)
 {
-	return NewObject<UMetaSoundOutputViewModel>(this, UMetaSoundOutputViewModel::StaticClass());
+	if (Builder)
+	{
+		const FMetaSoundFrontendDocumentBuilder& DocumentBuilder = Builder->GetConstBuilder();
+		const FMetasoundFrontendDocument& FrontendDocument = DocumentBuilder.GetConstDocumentChecked();
+		for (const FMetasoundFrontendClassInput& Input : FrontendDocument.RootGraph.GetDefaultInterface().Inputs)
+		{
+			if (Input.Name == VertexName)
+			{
+				if (UMetaSoundInputViewModel* InputViewModel = CreateInputViewModel(Input))
+				{
+					InputViewModel->SetIsInitialized(true);
+					UE_MVVM_BROADCAST_FIELD_VALUE_CHANGED(GetInputViewModels);
+				}
+				return;
+			}
+		}
+	}
 }
 
-void UMetaSoundInputViewModel::InitializeInput(UMetaSoundBuilderBase* InBuilder, const FMetasoundFrontendClassInput& InInput)
+void UMetaSoundViewModel::OnInputRemoved(const FName VertexName, FName DataType)
 {
-	ResetInput();
-	Builder = InBuilder;
+	if (InputViewModels.Remove(VertexName))
+	{
+		UE_MVVM_BROADCAST_FIELD_VALUE_CHANGED(GetInputViewModels);
+	}
+}
 
-	UE_MVVM_SET_PROPERTY_VALUE(InputName, InInput.Name);
-	UE_MVVM_SET_PROPERTY_VALUE(DataType, InInput.TypeName);
+void UMetaSoundViewModel::OnInputNameChanged(const FName OldName, const FName NewName)
+{
+	if (UMetaSoundInputViewModel* InputViewModel = Cast<UMetaSoundInputViewModel>(InputViewModels.FindAndRemoveChecked(OldName)))
+	{
+		InputViewModels.Emplace(NewName, MoveTemp(InputViewModel));
+		InputViewModel->OnInputNameChanged(NewName);
+	}
+}
 
-	if (const FMetasoundFrontendClassInputDefault* Default = Metasound::Frontend::FindPreferredPage(InInput.GetDefaults(), UMetaSoundSettings::GetPageOrder()); ensure(Default))
+void UMetaSoundViewModel::OnInputDataTypeChanged(const FName VertexName, const FName DataType)
+{
+	if (UMetaSoundInputViewModel* InputViewModel = *InputViewModels.Find(VertexName))
 	{
-		const FMetasoundFrontendLiteral& DefaultLiteral = Default->Literal;
+		InputViewModel->OnInputDataTypeChanged(DataType);
+	}
+}
 
-		UE_MVVM_SET_PROPERTY_VALUE(LiteralType, DefaultLiteral.GetType());
-		UE_MVVM_SET_PROPERTY_VALUE(Literal, DefaultLiteral);
-		UE_MVVM_SET_PROPERTY_VALUE(bIsArray, DefaultLiteral.IsArray());
+void UMetaSoundViewModel::OnInputDefaultChanged(const FName VertexName, FMetasoundFrontendLiteral LiteralValue, const FName PageName)
+{
+	if (UMetaSoundInputViewModel* InputViewModel = *InputViewModels.Find(VertexName))
+	{
+		InputViewModel->OnInputDefaultChanged(LiteralValue, PageName);
 	}
+}
 
-	SetIsInitialized(true);
+void UMetaSoundViewModel::OnInputInheritsDefaultChanged(const FName VertexName, const bool bInheritsDefault)
+{
+	if (UMetaSoundInputViewModel* InputViewModel = *InputViewModels.Find(VertexName))
+	{
+		InputViewModel->OnInputInheritsDefaultChanged(bInheritsDefault);
+	}
 }
 
-void UMetaSoundInputViewModel::ResetInput()
+void UMetaSoundViewModel::OnInputIsConstructorPinChanged(const FName VertexName, const bool bIsConstructorPin)
 {
-	SetIsInitialized(false);
-	Builder = nullptr;
+	if (UMetaSoundInputViewModel* InputViewModel = *InputViewModels.Find(VertexName))
+	{
+		InputViewModel->OnInputIsConstructorPinChanged(bIsConstructorPin);
+	}
+}
 
-	UE_MVVM_SET_PROPERTY_VALUE(InputName, FName());
-	UE_MVVM_SET_PROPERTY_VALUE(DataType, FName());
-	UE_MVVM_SET_PROPERTY_VALUE(LiteralType, EMetasoundFrontendLiteralType::Invalid);
-	UE_MVVM_SET_PROPERTY_VALUE(Literal, FMetasoundFrontendLiteral());
-	UE_MVVM_SET_PROPERTY_VALUE(bIsArray, false);
+void UMetaSoundViewModel::OnOutputAdded(const FName VertexName, FName DataType)
+{
+	if (Builder)
+	{
+		const FMetaSoundFrontendDocumentBuilder& DocumentBuilder = Builder->GetConstBuilder();
+		const FMetasoundFrontendDocument& FrontendDocument = DocumentBuilder.GetConstDocumentChecked();
+		for (const FMetasoundFrontendClassOutput& Output : FrontendDocument.RootGraph.GetDefaultInterface().Outputs)
+		{
+			if (Output.Name == VertexName)
+			{
+				if (UMetaSoundOutputViewModel* OutputViewModel = CreateOutputViewModel(Output))
+				{
+					OutputViewModel->SetIsInitialized(true);
+					UE_MVVM_BROADCAST_FIELD_VALUE_CHANGED(GetOutputViewModels);
+				}
+				return;
+			}
+		}
+	}
 }
 
-void UMetaSoundInputViewModel::SetIsInitialized(const bool bInIsInitialized)
+void UMetaSoundViewModel::OnOutputRemoved(const FName VertexName, FName DataType)
 {
-	UE_MVVM_SET_PROPERTY_VALUE(bIsInitialized, bInIsInitialized);
+	if (OutputViewModels.Remove(VertexName))
+	{
+		UE_MVVM_BROADCAST_FIELD_VALUE_CHANGED(GetOutputViewModels);
+	}
 }
 
-void UMetaSoundInputViewModel::SetInputName(const FName& InInputName)
+void UMetaSoundViewModel::OnOutputNameChanged(const FName OldName, const FName NewName)
+{
+	if (UMetaSoundOutputViewModel* OutputViewModel = OutputViewModels.FindAndRemoveChecked(OldName))
+	{
+		OutputViewModels.Emplace(NewName, MoveTemp(OutputViewModel));
+		OutputViewModel->OnOutputNameChanged(NewName);
+	}
+}
+
+void UMetaSoundViewModel::OnOutputDataTypeChanged(const FName VertexName, const FName DataType)
+{
+	if (UMetaSoundOutputViewModel* OutputViewModel = *OutputViewModels.Find(VertexName))
+	{
+		OutputViewModel->OnOutputDataTypeChanged(DataType);
+	}
+}
+
+void UMetaSoundViewModel::OnOutputIsConstructorPinChanged(const FName VertexName, const bool bIsConstructorPin)
+{
+	if (UMetaSoundOutputViewModel* OutputViewModel = *OutputViewModels.Find(VertexName))
+	{
+		OutputViewModel->OnOutputIsConstructorPinChanged(bIsConstructorPin);
+	}
+}
+
+void UMetaSoundViewModel::CreateMemberViewModels()
 {
 	if (!Builder)
 	{
-		UE_LOGFMT(LogTechAudioToolsMetaSound, Log, "Failed to set input name {new} for {old}. Builder was null.", InInputName, InputName);
+		UE_LOGFMT(LogTechAudioTools, Warning, "Failed to create member viewmodels. Builder was null.");
 		return;
 	}
 
-	const FName OldName = InputName;
-	if (UE_MVVM_SET_PROPERTY_VALUE(InputName, InInputName))
+	const FMetaSoundFrontendDocumentBuilder& DocumentBuilder = Builder->GetConstBuilder();
+	const FMetasoundFrontendDocument& FrontendDocument = DocumentBuilder.GetConstDocumentChecked();
+
+	for (const FMetasoundFrontendClassInput& Input : FrontendDocument.RootGraph.GetDefaultInterface().Inputs)
 	{
-		EMetaSoundBuilderResult Result;
-		Builder->SetGraphInputName(OldName, InInputName, Result);
+		if (UMetaSoundInputViewModel* InputViewModel = CreateInputViewModel(Input))
+		{
+			InputViewModel->SetIsInitialized(true);
+		}
+	}
+	UE_MVVM_BROADCAST_FIELD_VALUE_CHANGED(GetInputViewModels);
 
-		if (Result != EMetaSoundBuilderResult::Succeeded)
+	for (const FMetasoundFrontendClassOutput& Output : FrontendDocument.RootGraph.GetDefaultInterface().Outputs)
+	{
+		if (UMetaSoundOutputViewModel* OutputViewModel = CreateOutputViewModel(Output))
 		{
-			UE_LOGFMT(LogTechAudioToolsMetaSound, Log, "Failed to set input name {new} for {old}. Builder was null.", InInputName, InputName);
+			OutputViewModel->SetIsInitialized(true);
 		}
 	}
+	UE_MVVM_BROADCAST_FIELD_VALUE_CHANGED(GetOutputViewModels);
 }
 
-void UMetaSoundInputViewModel::SetDataType(const FName& InDataType)
+UMetaSoundInputViewModel* UMetaSoundViewModel::CreateInputViewModel(const FMetasoundFrontendClassInput& InInput)
 {
-	if (!Builder)
+	if (UMetaSoundInputViewModel* InputViewModel = NewObject<UMetaSoundInputViewModel>(this, GetInputViewModelClass()))
 	{
-		UE_LOGFMT(LogTechAudioToolsMetaSound, Log, "Failed to set data type for {input}. Builder was null.", InputName);
-		return;
+		FName PageName;
+		FMetasoundFrontendLiteral DefaultLiteral;
+		if (const FMetasoundFrontendClassInputDefault* Default = Metasound::Frontend::FindPreferredPage(InInput.GetDefaults(), UMetaSoundSettings::GetPageOrder()))
+		{
+			if (const UMetaSoundSettings* Settings = GetDefault<UMetaSoundSettings>())
+			{
+				if (const FMetaSoundPageSettings* PageSettings = Settings->FindPageSettings(Default->PageID))
+				{
+					PageName = PageSettings->Name;
+				}
+			}
+
+			DefaultLiteral = Default->Literal;
+		}
+
+		bool bInheritsDefault = false;
+		if (Builder)
+		{
+			const FMetaSoundFrontendDocumentBuilder& DocumentBuilder = Builder->GetConstBuilder();
+			const FMetasoundFrontendDocument& FrontendDocument = DocumentBuilder.GetConstDocumentChecked();
+			bInheritsDefault = FrontendDocument.RootGraph.PresetOptions.InputsInheritingDefault.Contains(InInput.Name);
+		}
+
+		InputViewModel->OnInputNameChanged(InInput.Name);
+		InputViewModel->OnInputDataTypeChanged(InInput.TypeName);
+		InputViewModel->OnInputDefaultChanged(DefaultLiteral, PageName);
+		InputViewModel->OnInputInheritsDefaultChanged(bInheritsDefault);
+		InputViewModel->OnInputIsConstructorPinChanged(InInput.AccessType == EMetasoundFrontendVertexAccessType::Value);
+
+		return InputViewModels.Emplace(InInput.Name, InputViewModel);
+	}
+
+	return nullptr;
+}
+
+UMetaSoundOutputViewModel* UMetaSoundViewModel::CreateOutputViewModel(const FMetasoundFrontendClassOutput& InOutput)
+{
+	if (UMetaSoundOutputViewModel* OutputViewModel = NewObject<UMetaSoundOutputViewModel>(this, GetOutputViewModelClass()))
+	{
+		OutputViewModel->OnOutputNameChanged(InOutput.Name);
+		OutputViewModel->OnOutputDataTypeChanged(InOutput.TypeName);
+		OutputViewModel->OnOutputIsConstructorPinChanged(InOutput.AccessType == EMetasoundFrontendVertexAccessType::Value);
+
+		return OutputViewModels.Emplace(InOutput.Name, OutputViewModel);
+	}
+
+	return nullptr;
+}
+
+TSubclassOf<UMetaSoundInputViewModel> UMetaSoundViewModel::GetInputViewModelClass() const
+{
+	return UMetaSoundInputViewModel::StaticClass();
+}
+
+TSubclassOf<UMetaSoundOutputViewModel> UMetaSoundViewModel::GetOutputViewModelClass() const
+{
+	return UMetaSoundOutputViewModel::StaticClass();
+}
+
+void UMetaSoundInputViewModel::SetIsInitialized(const bool bInIsInitialized)
+{
+	UE_MVVM_SET_PROPERTY_VALUE(bIsInitialized, bInIsInitialized);
+}
+
+void UMetaSoundInputViewModel::SetInputName(const FName& NewName)
+{
+	const FName OldName = InputName;
+	if (const UMetaSoundViewModel* MetaSoundViewModel = Cast<UMetaSoundViewModel>(GetOuter()))
+	{
+		MetaSoundViewModel->SetInputName(OldName, NewName);
 	}
+}
 
-	FMetaSoundFrontendDocumentBuilder& DocumentBuilder = Builder->GetBuilder();
-	if (!DocumentBuilder.IsValid())
+void UMetaSoundInputViewModel::SetDataType(const FName& InDataType)
+{
+	if (!Metasound::Frontend::IDataTypeRegistry::Get().IsRegistered(InDataType))
 	{
-		UE_LOGFMT(LogTechAudioToolsMetaSound, Log, "Failed to set data type for {input}. DocumentBuilder was invalid.", InputName);
+		UE_MVVM_BROADCAST_FIELD_VALUE_CHANGED(DataType);
+		UE_LOGFMT(LogTechAudioTools, Warning, "Failed to set data type for {InputName}. '{InDataType}' is not a registered data type.", InputName, InDataType);
 		return;
 	}
 
-	if (UE_MVVM_SET_PROPERTY_VALUE(DataType, InDataType))
+	if (const UMetaSoundViewModel* MetaSoundViewModel = Cast<UMetaSoundViewModel>(GetOuter()))
 	{
-		DocumentBuilder.SetGraphInputDataType(InputName, InDataType);
+		MetaSoundViewModel->SetInputDataType(InputName, InDataType);
 	}
 }
 
 void UMetaSoundInputViewModel::SetIsArray(const bool bInIsArray)
 {
-	if (UE_MVVM_SET_PROPERTY_VALUE(bIsArray, bInIsArray))
+	if (UMetaSoundViewModelConversionFunctions::IsArrayType(DataType))
 	{
-		SetDataType(TechAudioTools::MetaSoundViewModel::GetAdjustedDataType(DataType, bInIsArray));
+		SetDataType(TechAudioTools::MetaSound::GetAdjustedDataType(DataType, bInIsArray));
 	}
 }
 
-FText UMetaSoundInputViewModel::GetLiteralValueAsText() const
+void UMetaSoundInputViewModel::SetIsConstructorPin(const bool bInIsConstructorPin)
 {
-	return FText::FromString(Literal.ToString());
+	if (!UMetaSoundViewModelConversionFunctions::IsConstructorType(DataType))
+	{
+		UE_MVVM_BROADCAST_FIELD_VALUE_CHANGED(bIsConstructorPin);
+		UE_LOGFMT(LogTechAudioTools, Warning, "Failed to set bIsConstructorPin on {InputName}. Data type '{DataType}' cannot be a constructor pin.", InputName, DataType);
+		return;
+	}
+
+	if (const UMetaSoundViewModel* MetaSoundViewModel = Cast<UMetaSoundViewModel>(GetOuter()))
+	{
+		MetaSoundViewModel->SetInputIsConstructorPin(InputName, bInIsConstructorPin);
+	}
 }
 
 void UMetaSoundInputViewModel::SetLiteral(const FMetasoundFrontendLiteral& InLiteral)
 {
-	if (!Builder)
+	if (const UMetaSoundViewModel* MetaSoundViewModel = Cast<UMetaSoundViewModel>(GetOuter()))
 	{
-		UE_LOGFMT(LogTechAudioToolsMetaSound, Log, "Failed to set literal for {input}. Builder was null.", InputName);
-		return;
+		MetaSoundViewModel->SetInputDefaultLiteral(InputName, InLiteral);
 	}
+}
 
-	if (UE_MVVM_SET_PROPERTY_VALUE(Literal, InLiteral))
+void UMetaSoundInputViewModel::SetOverridesDefault(const bool bInOverridesDefaultValue)
+{
+	if (const UMetaSoundViewModel* MetaSoundViewModel = Cast<UMetaSoundViewModel>(GetOuter()))
 	{
-		EMetaSoundBuilderResult Result;
-		Builder->SetGraphInputDefault(InputName, InLiteral, Result);
-		UE_MVVM_BROADCAST_FIELD_VALUE_CHANGED(GetLiteralValueAsText);
+		MetaSoundViewModel->SetInputOverridesDefault(InputName, bInOverridesDefaultValue);
 	}
 }
 
-void UMetaSoundOutputViewModel::InitializeOutput(UMetaSoundBuilderBase* InBuilder, const FMetasoundFrontendClassOutput& InOutput)
+void UMetaSoundInputViewModel::OnInputNameChanged(const FName& NewName)
 {
-	ResetOutput();
-	Builder = InBuilder;
+	UE_MVVM_SET_PROPERTY_VALUE(InputName, NewName);
+}
 
-	UE_MVVM_SET_PROPERTY_VALUE(OutputName, InOutput.Name);
-	UE_MVVM_SET_PROPERTY_VALUE(DataType, InOutput.TypeName);
-	SetMetaSoundOutput(FMetaSoundOutput());
-	SetIsInitialized(true);
+void UMetaSoundInputViewModel::OnInputDataTypeChanged(const FName& NewDataType)
+{
+	if (UE_MVVM_SET_PROPERTY_VALUE(DataType, NewDataType))
+	{
+		UE_MVVM_BROADCAST_FIELD_VALUE_CHANGED(LiteralType);
+		UE_MVVM_BROADCAST_FIELD_VALUE_CHANGED(bIsArray);
+	}
+}
+
+void UMetaSoundInputViewModel::OnInputDefaultChanged(const FMetasoundFrontendLiteral& LiteralValue, const FName& PageName)
+{
+	UE_MVVM_SET_PROPERTY_VALUE(Literal, LiteralValue);
 }
 
-void UMetaSoundOutputViewModel::ResetOutput()
+void UMetaSoundInputViewModel::OnInputInheritsDefaultChanged(const bool bInheritsDefault)
 {
-	SetMetaSoundOutput(FMetaSoundOutput());
+	UE_MVVM_SET_PROPERTY_VALUE(bOverridesDefault, !bInheritsDefault);
+}
+
+void UMetaSoundInputViewModel::OnInputIsConstructorPinChanged(const bool bNewIsConstructorPin)
+{
+	UE_MVVM_SET_PROPERTY_VALUE(bIsConstructorPin, bNewIsConstructorPin);
 }
 
 void UMetaSoundOutputViewModel::SetIsInitialized(const bool bInIsInitialized)
@@ -293,66 +536,67 @@ void UMetaSoundOutputViewModel::SetIsInitialized(const bool bInIsInitialized)
 	UE_MVVM_SET_PROPERTY_VALUE(bIsInitialized, bInIsInitialized);
 }
 
-void UMetaSoundOutputViewModel::SetOutputName(const FName& InOutputName)
+void UMetaSoundOutputViewModel::SetOutputName(const FName& NewName) const
 {
-	if (!Builder)
-	{
-		UE_LOGFMT(LogTechAudioToolsMetaSound, Log, "Failed to set output name for {output}. Builder was null.", OutputName);
-		return;
-	}
-
 	const FName OldName = OutputName;
-	if (UE_MVVM_SET_PROPERTY_VALUE(OutputName, InOutputName))
+	if (const UMetaSoundViewModel* MetaSoundViewModel = Cast<UMetaSoundViewModel>(GetOuter()))
 	{
-		EMetaSoundBuilderResult Result;
-		Builder->SetGraphOutputName(OldName, InOutputName, Result);
-
-		if (Result != EMetaSoundBuilderResult::Succeeded)
-		{
-			UE_LOGFMT(LogTechAudioToolsMetaSound, Log, "Failed to set output name for {output}", OutputName);
-		}
+		MetaSoundViewModel->SetOutputName(OldName, NewName);
 	}
 }
 
 void UMetaSoundOutputViewModel::SetDataType(const FName& InDataType)
 {
-	if (!Builder)
+	if (!Metasound::Frontend::IDataTypeRegistry::Get().IsRegistered(InDataType))
 	{
-		UE_LOGFMT(LogTechAudioToolsMetaSound, Log, "Failed to set data type for {output}. Builder was null.", OutputName);
+		UE_MVVM_BROADCAST_FIELD_VALUE_CHANGED(DataType);
+		UE_LOGFMT(LogTechAudioTools, Warning, "Failed to set data type for {OutputName}. '{InDataType}' is not a registered data type.", OutputName, InDataType);
 		return;
 	}
 
-	FMetaSoundFrontendDocumentBuilder& DocumentBuilder = Builder->GetBuilder();
-	if (!DocumentBuilder.IsValid())
+	if (const UMetaSoundViewModel* MetaSoundViewModel = Cast<UMetaSoundViewModel>(GetOuter()))
 	{
-		UE_LOGFMT(LogTechAudioToolsMetaSound, Log, "Failed to set data type for {output}. DocumentBuilder was invalid.", OutputName);
-		return;
+		MetaSoundViewModel->SetOutputDataType(OutputName, InDataType);
 	}
+}
 
-	if (UE_MVVM_SET_PROPERTY_VALUE(DataType, InDataType))
+void UMetaSoundOutputViewModel::SetIsArray(const bool bInIsArray)
+{
+	if (UMetaSoundViewModelConversionFunctions::IsArrayType(DataType))
 	{
-		DocumentBuilder.SetGraphOutputDataType(OutputName, InDataType);
+		SetDataType(TechAudioTools::MetaSound::GetAdjustedDataType(DataType, bInIsArray));
 	}
 }
 
-void UMetaSoundOutputViewModel::SetIsArray(const bool bInIsArray)
+void UMetaSoundOutputViewModel::SetIsConstructorPin(const bool bInIsConstructorPin)
 {
-	if (UE_MVVM_SET_PROPERTY_VALUE(bIsArray, bInIsArray))
+	if (!UMetaSoundViewModelConversionFunctions::IsConstructorType(DataType))
 	{
-		SetDataType(TechAudioTools::MetaSoundViewModel::GetAdjustedDataType(DataType, bInIsArray));
+		UE_MVVM_BROADCAST_FIELD_VALUE_CHANGED(bIsConstructorPin);
+		UE_LOGFMT(LogTechAudioTools, Warning, "Failed to set bIsConstructorPin on {OutputName}. '{DataType}' cannot be a constructor pin.", OutputName, DataType);
+		return;
+	}
+
+	if (const UMetaSoundViewModel* MetaSoundViewModel = Cast<UMetaSoundViewModel>(GetOuter()))
+	{
+		MetaSoundViewModel->SetOutputIsConstructorPin(OutputName, bInIsConstructorPin);
 	}
 }
 
-void UMetaSoundOutputViewModel::SetMetaSoundOutput(const FMetaSoundOutput& InMetaSoundOutput)
+void UMetaSoundOutputViewModel::OnOutputNameChanged(const FName& NewName)
 {
-	MetaSoundOutput = InMetaSoundOutput;
-	UE_MVVM_BROADCAST_FIELD_VALUE_CHANGED(MetaSoundOutput);
+	UE_MVVM_SET_PROPERTY_VALUE(OutputName, NewName);
 }
 
-void UMetaSoundOutputViewModel::OnOutputValueChanged(const FName InOutputName, const FMetaSoundOutput& InOutput)
+void UMetaSoundOutputViewModel::OnOutputDataTypeChanged(const FName& NewDataType)
 {
-	if (OutputName == InOutputName)
+	if (UE_MVVM_SET_PROPERTY_VALUE(DataType, NewDataType))
 	{
-		SetMetaSoundOutput(InOutput);
+		UE_MVVM_BROADCAST_FIELD_VALUE_CHANGED(bIsArray);
 	}
 }
+
+void UMetaSoundOutputViewModel::OnOutputIsConstructorPinChanged(const bool bNewIsConstructorPin)
+{
+	UE_MVVM_SET_PROPERTY_VALUE(bIsConstructorPin, bNewIsConstructorPin);
+}
diff --git a/Engine/Plugins/Experimental/TechAudioTools/Source/TechAudioToolsMetaSound/Private/ViewModels/MetaSoundViewModelConversionFunctions.cpp b/Engine/Plugins/Experimental/TechAudioTools/Source/TechAudioToolsMetaSound/Private/ViewModels/MetaSoundViewModelConversionFunctions.cpp
--- a/Engine/Plugins/Experimental/TechAudioTools/Source/TechAudioToolsMetaSound/Private/ViewModels/MetaSoundViewModelConversionFunctions.cpp
+++ b/Engine/Plugins/Experimental/TechAudioTools/Source/TechAudioToolsMetaSound/Private/ViewModels/MetaSoundViewModelConversionFunctions.cpp
@@ -2,6 +2,9 @@
 
 #include "ViewModels/MetaSoundViewModelConversionFunctions.h"
 
+#include "MetasoundFrontendSearchEngine.h"
+#include "TechAudioToolsLog.h"
+#include "Logging/StructuredLog.h"
 #include "ViewModels/MetaSoundViewModel.h"
 
 #include UE_INLINE_GENERATED_CPP_BY_NAME(MetaSoundViewModelConversionFunctions)
@@ -31,3 +34,42 @@ UMetaSoundOutputViewModel* UMetaSoundViewModelConversionFunctions::FindOutputVie
 
 	return nullptr;
 }
+
+FText UMetaSoundViewModelConversionFunctions::GetLiteralValueAsText(const FMetasoundFrontendLiteral& Literal)
+{
+	return FText::FromString(Literal.ToString());
+}
+
+bool UMetaSoundViewModelConversionFunctions::IsInterfaceMember(const FName& MemberName, const bool bInvert)
+{
+	FName InterfaceNamespace;
+	FName ParamName;
+	Audio::FParameterPath::SplitName(MemberName, InterfaceNamespace, ParamName);
+
+	FMetasoundFrontendInterface FoundInterface;
+	const bool bIsValidNamespace = Metasound::Frontend::ISearchEngine::Get().FindInterfaceWithHighestVersion(InterfaceNamespace, FoundInterface);
+	const bool bResult = !InterfaceNamespace.IsNone() && bIsValidNamespace;
+
+	return bInvert ? !bResult : bResult;
+}
+
+bool UMetaSoundViewModelConversionFunctions::IsArrayType(const FName& DataType, const bool bInvert)
+{
+	const FName AdjustedDataType = TechAudioTools::MetaSound::GetAdjustedDataType(DataType, true);
+	const bool bIsValidDataType = Metasound::Frontend::IDataTypeRegistry::Get().IsRegistered(AdjustedDataType);
+	return bInvert ? !bIsValidDataType : bIsValidDataType;
+}
+
+bool UMetaSoundViewModelConversionFunctions::IsConstructorType(const FName& DataType, const bool bInvert)
+{
+	using namespace Metasound::Frontend;
+
+	FDataTypeRegistryInfo Info;
+	if (!IDataTypeRegistry::Get().GetDataTypeInfo(DataType, Info))
+	{
+		UE_LOGFMT(LogTechAudioTools, Warning, "'{DataType}' is not a registered data type.", DataType);
+		return false;
+	}
+
+	return bInvert ? !Info.bIsConstructorType : Info.bIsConstructorType;
+}
diff --git a/Engine/Plugins/Experimental/TechAudioTools/Source/TechAudioToolsMetaSound/Public/ViewModels/MetaSoundViewModel.h b/Engine/Plugins/Experimental/TechAudioTools/Source/TechAudioToolsMetaSound/Public/ViewModels/MetaSoundViewModel.h
--- a/Engine/Plugins/Experimental/TechAudioTools/Source/TechAudioToolsMetaSound/Public/ViewModels/MetaSoundViewModel.h
+++ b/Engine/Plugins/Experimental/TechAudioTools/Source/TechAudioToolsMetaSound/Public/ViewModels/MetaSoundViewModel.h
@@ -3,15 +3,12 @@
 #pragma once
 
 #include "MetasoundFrontendDocumentBuilder.h"
-#include "MetasoundGeneratorHandle.h"
 #include "MVVMViewModelBase.h"
 
 #include "MetaSoundViewModel.generated.h"
 
 #define UE_API TECHAUDIOTOOLSMETASOUND_API
 
-TECHAUDIOTOOLSMETASOUND_API DECLARE_LOG_CATEGORY_EXTERN(LogTechAudioToolsMetaSound, Log, All);
-
 class UMetaSoundBuilderBase;
 class UMetaSoundInputViewModel;
 class UMetaSoundOutputViewModel;
@@ -26,19 +23,19 @@ class UMetaSoundViewModel : public UMVVMViewModelBase
 {
 	GENERATED_BODY()
 
-private:
+protected:
 	// True if this MetaSound Viewmodel has been initialized.
-	UPROPERTY(BlueprintReadOnly, FieldNotify, Category = "MetaSound Viewmodel", meta = (AllowPrivateAccess))
+	UPROPERTY(BlueprintReadOnly, FieldNotify, Getter = "IsInitialized", Category = "MetaSound Viewmodel", meta = (AllowPrivateAccess))
 	bool bIsInitialized = false;
 
 	// True if the initialized MetaSound is a preset.
-	UPROPERTY(BlueprintReadOnly, FieldNotify, Category = "MetaSound Viewmodel", meta = (AllowPrivateAccess))
+	UPROPERTY(BlueprintReadOnly, FieldNotify, Getter = "IsPreset", Category = "MetaSound Viewmodel", meta = (AllowPrivateAccess))
 	bool bIsPreset = false;
 
 public:
-	// Returns the object name of the initialized builder as text.
+	// Returns the object name of the initialized builder.
 	UFUNCTION(BlueprintCallable, FieldNotify, Category = "MetaSound Viewmodel")
-	UE_API FText GetBuilderNameAsText() const;
+	UE_API FString GetBuilderName() const { return Builder.GetName(); }
 
 	// Contains MetaSound Input Viewmodels for each input of the initialized MetaSound.
 	UFUNCTION(BlueprintCallable, FieldNotify, DisplayName = "Get Input Viewmodels", Category = "MetaSound Viewmodel")
@@ -60,27 +57,54 @@ public:
 	UFUNCTION(BlueprintCallable, Category = "Audio|MetaSound Viewmodel")
 	UE_API virtual void Reset();
 
-	// Returns a reference to the initialized MetaSound's Builder.
-	UFUNCTION(BlueprintCallable, Category = "Audio|MetaSound Viewmodel")
-	UMetaSoundBuilderBase* GetBuilder() const { return Builder; }
-
 	bool IsInitialized() const { return bIsInitialized; }
-	void SetIsInitialized(const bool bInIsInitialized) { UE_MVVM_SET_PROPERTY_VALUE(bIsInitialized, bInIsInitialized); }
+	bool IsPreset() const { return bIsPreset; }
 
-protected:
-	UE_API virtual void InitializeProperties(const FMetasoundFrontendDocument& FrontendDocument);
-	UE_API virtual void ResetProperties();
+	UE_API void SetIsInitialized(const bool bInIsInitialized);
 
-	// Called upon initialization. Creates viewmodel instances for all inputs and outputs of the initialized MetaSound.
-	UE_API void CreateMemberViewModels();
+	UE_API void SetInputName(const FName& OldName, const FName& NewName) const;
+	UE_API void SetInputDataType(const FName& InputName, const FName& DataType) const;
+	UE_API void SetInputDefaultLiteral(const FName& InputName, const FMetasoundFrontendLiteral& DefaultLiteral) const;
+	UE_API void SetInputOverridesDefault(const FName& InputName, bool bOverridesDefault) const;
+	UE_API void SetInputIsConstructorPin(const FName& InputName, bool bIsConstructorPin) const;
+
+	UE_API void SetOutputName(const FName& OldName, const FName& NewName) const;
+	UE_API void SetOutputDataType(const FName& OutputName, const FName& DataType) const;
+	UE_API void SetOutputIsConstructorPin(const FName& OutputName, bool bIsConstructorPin) const;
+
+	UFUNCTION()
+	UE_API void OnInputAdded(FName VertexName, FName DataType);
+	UFUNCTION()
+	UE_API void OnInputRemoved(FName VertexName, FName DataType);
+	UFUNCTION()
+	UE_API void OnInputNameChanged(FName OldName, FName NewName);
+	UFUNCTION()
+	UE_API void OnInputDataTypeChanged(FName VertexName, FName DataType);
+	UFUNCTION()
+	UE_API void OnInputDefaultChanged(FName VertexName, FMetasoundFrontendLiteral LiteralValue, FName PageName);
+	UFUNCTION()
+	UE_API void OnInputInheritsDefaultChanged(FName VertexName, bool bInheritsDefault);
+	UFUNCTION()
+	UE_API void OnInputIsConstructorPinChanged(FName VertexName, bool bIsConstructorPin);
+
+	UFUNCTION()
+	UE_API void OnOutputAdded(FName VertexName, FName DataType);
+	UFUNCTION()
+	UE_API void OnOutputRemoved(FName VertexName, FName DataType);
+	UFUNCTION()
+	UE_API void OnOutputNameChanged(FName OldName, FName NewName);
+	UFUNCTION()
+	UE_API void OnOutputDataTypeChanged(FName VertexName, FName DataType);
+	UFUNCTION()
+	UE_API void OnOutputIsConstructorPinChanged(FName VertexName, bool bIsConstructorPin);
 
-	// Creates a single MetaSoundInputViewModel instance for the given input.
-	UE_API virtual void CreateInputViewModel(const FMetasoundFrontendClassInput& InInput);
-    UE_API virtual UMetaSoundInputViewModel* CreateInputViewModelInstance();
+protected:
+	UE_API void CreateMemberViewModels();
+	UE_API virtual UMetaSoundInputViewModel* CreateInputViewModel(const FMetasoundFrontendClassInput& InInput);
+	UE_API virtual UMetaSoundOutputViewModel* CreateOutputViewModel(const FMetasoundFrontendClassOutput& InOutput);
 
-	// Creates a single MetaSoundOutputViewModel instance for the given output.
-	UE_API void CreateOutputViewModel(const FMetasoundFrontendClassOutput& InOutput);
-	UE_API virtual UMetaSoundOutputViewModel* CreateOutputViewModelInstance();
+	virtual TSubclassOf<UMetaSoundInputViewModel> GetInputViewModelClass() const;
+	virtual TSubclassOf<UMetaSoundOutputViewModel> GetOutputViewModelClass() const;
 
 	UPROPERTY(Transient)
 	TObjectPtr<UMetaSoundBuilderBase> Builder;
@@ -103,48 +127,60 @@ class UMetaSoundInputViewModel : public UMVVMViewModelBase
 
 protected:
 	// True if this MetaSoundInputViewModel has been initialized.
-	UPROPERTY(BlueprintReadOnly, FieldNotify, Category = "MetaSound Input", meta = (AllowPrivateAccess))
+	UPROPERTY(BlueprintReadOnly, FieldNotify, Getter = "IsInitialized", Category = "MetaSound Input", meta = (AllowPrivateAccess))
 	bool bIsInitialized = false;
 
-	// Sets the name of the initialized MetaSound input.
-	UPROPERTY(BlueprintReadWrite, EditAnywhere, FieldNotify, Setter, Category = "MetaSound Input", meta = (AllowPrivateAccess))
+	// The name of the initialized MetaSound input.
+	UPROPERTY(BlueprintReadWrite, EditAnywhere, FieldNotify, Getter, Setter, Category = "MetaSound Input", meta = (AllowPrivateAccess))
 	FName InputName;
 
-	// Returns the data type of the initialized MetaSound input.
-	UPROPERTY(BlueprintReadWrite, FieldNotify, Setter, Category = "MetaSound Input", meta = (AllowPrivateAccess))
+	// The data type of the initialized MetaSound input.
+	UPROPERTY(BlueprintReadWrite, FieldNotify, Getter, Setter, Category = "MetaSound Input", meta = (AllowPrivateAccess))
 	FName DataType;
 
 	// True if the initialized MetaSound input is an array.
-	UPROPERTY(BlueprintReadWrite, FieldNotify, Setter = "SetIsArray", Category = "MetaSound Input", meta = (AllowPrivateAccess))
+	UPROPERTY(BlueprintReadWrite, FieldNotify, Getter = "IsArray", Setter = "SetIsArray", Category = "MetaSound Input", meta = (AllowPrivateAccess))
 	bool bIsArray;
 
+	// True if the initialized MetaSound input is a constructor pin.
+	UPROPERTY(BlueprintReadWrite, FieldNotify, Getter = "IsConstructorPin", Setter = "SetIsConstructorPin", Category = "MetaSound Input", meta = (AllowPrivateAccess))
+	bool bIsConstructorPin;
+
 	// The MetaSound Literal belonging to the initialized MetaSound input.
-	UPROPERTY(BlueprintReadWrite, FieldNotify, Setter, Category = "MetaSound Input", meta = (AllowPrivateAccess))
+	UPROPERTY(BlueprintReadWrite, FieldNotify, Getter, Setter, Category = "MetaSound Input", meta = (AllowPrivateAccess))
 	FMetasoundFrontendLiteral Literal;
 
 	// Returns the Literal Type belonging to the initialized MetaSound input.
-	UPROPERTY(BlueprintReadOnly, FieldNotify, Category = "MetaSound Input", meta = (AllowPrivateAccess))
+	UPROPERTY(BlueprintReadOnly, FieldNotify, Getter, Category = "MetaSound Input", meta = (AllowPrivateAccess))
 	EMetasoundFrontendLiteralType LiteralType;
 
-	UPROPERTY(Transient)
-	TObjectPtr<UMetaSoundBuilderBase> Builder;
+	// True if the initialized MetaSound input overrides its value from a parent graph.
+	UPROPERTY(BlueprintReadWrite, FieldNotify, Getter = "OverridesDefault", Setter = "SetOverridesDefault", Category = "MetaSound Input", meta = (AllowPrivateAccess))
+	bool bOverridesDefault;
 
 public:
-	UE_API virtual void InitializeInput(UMetaSoundBuilderBase* InBuilder, const FMetasoundFrontendClassInput& InInput);
-	UE_API virtual void ResetInput();
-
+	bool IsInitialized() const { return bIsInitialized; }
 	FName GetInputName() const { return InputName; }
+	FName GetDataType() const { return DataType; }
+	bool IsArray() const { return DataType.ToString().EndsWith(METASOUND_DATA_TYPE_NAME_ARRAY_TYPE_SPECIFIER); }
+	bool IsConstructorPin() const { return bIsConstructorPin; }
+	FMetasoundFrontendLiteral GetLiteral() const { return Literal; }
+	EMetasoundFrontendLiteralType GetLiteralType() const { return Literal.GetType(); }
+	bool OverridesDefault() const { return bOverridesDefault; }
 
 	UE_API void SetIsInitialized(const bool bInIsInitialized);
-	UE_API void SetInputName(const FName& InInputName);
+	UE_API void SetInputName(const FName& NewName);
 	UE_API void SetDataType(const FName& InDataType);
-	UE_API void SetIsArray(const bool bInIsArray);
-
-	// Returns the value of this input's MetaSound Literal as a text value.
-	UFUNCTION(BlueprintCallable, FieldNotify, Category = "Audio|MetaSound Input Viewmodel")
-	UE_API FText GetLiteralValueAsText() const;
-
+	UE_API void SetIsArray(bool bInIsArray);
+	UE_API void SetIsConstructorPin(bool bInIsConstructorPin);
 	UE_API void SetLiteral(const FMetasoundFrontendLiteral& InLiteral);
+	UE_API void SetOverridesDefault(bool bInOverridesDefaultValue);
+
+	UE_API void OnInputNameChanged(const FName& NewName);
+	UE_API void OnInputDataTypeChanged(const FName& NewDataType);
+	UE_API void OnInputDefaultChanged(const FMetasoundFrontendLiteral& LiteralValue, const FName& PageName);
+	UE_API void OnInputInheritsDefaultChanged(const bool bInheritsDefault);
+	UE_API void OnInputIsConstructorPinChanged(const bool bNewIsConstructorPin);
 };
 
 /**
@@ -161,39 +197,38 @@ protected:
 	UPROPERTY(BlueprintReadOnly, FieldNotify, Category = "MetaSound Output", meta = (AllowPrivateAccess))
 	bool bIsInitialized = false;
 
-	// Sets the name of the initialized MetaSound output.
+	// The name of the initialized MetaSound output.
 	UPROPERTY(BlueprintReadWrite, EditAnywhere, FieldNotify, Setter, Category = "MetaSound Output", meta = (AllowPrivateAccess))
 	FName OutputName;
 
-	// Returns the data type of the initialized MetaSound output.
+	// The data type of the initialized MetaSound output.
 	UPROPERTY(BlueprintReadWrite, FieldNotify, Setter, Category = "MetaSound Output", meta = (AllowPrivateAccess))
 	FName DataType;
 
 	// True if the initialized MetaSound output is an array.
-	UPROPERTY(BlueprintReadWrite, FieldNotify, Setter = "SetIsArray", Category = "MetaSound Output", meta = (AllowPrivateAccess))
+	UPROPERTY(BlueprintReadWrite, FieldNotify, Getter = "IsArray", Setter = "SetIsArray", Category = "MetaSound Output", meta = (AllowPrivateAccess))
 	bool bIsArray;
 
-	// The MetaSound Output belonging to the initialized MetaSound output.
-	UPROPERTY(BlueprintReadWrite, FieldNotify, Setter, DisplayName = "MetaSound Output", Category = "MetaSound Output", meta = (AllowPrivateAccess))
-	FMetaSoundOutput MetaSoundOutput;
-
-	UPROPERTY(Transient)
-	TObjectPtr<UMetaSoundBuilderBase> Builder;
+	// True if the initialized MetaSound output is a constructor pin.
+	UPROPERTY(BlueprintReadWrite, FieldNotify, Getter = "IsConstructorPin", Setter = "SetIsConstructorPin", Category = "MetaSound Output", meta = (AllowPrivateAccess))
+	bool bIsConstructorPin;
 
 public:
-	UE_API virtual void InitializeOutput(UMetaSoundBuilderBase* InBuilder, const FMetasoundFrontendClassOutput& InOutput);
-	UE_API virtual void ResetOutput();
-
+	bool IsInitialized() const { return bIsInitialized; }
 	FName GetOutputName() const { return OutputName; }
+	FName GetDataType() const { return DataType; }
+	bool IsArray() const { return DataType.ToString().EndsWith(METASOUND_DATA_TYPE_NAME_ARRAY_TYPE_SPECIFIER); }
+	bool IsConstructorPin() const { return bIsConstructorPin; }
 
 	UE_API void SetIsInitialized(const bool bInIsInitialized);
-	UE_API void SetOutputName(const FName& InOutputName);
+	UE_API void SetOutputName(const FName& NewName) const;
 	UE_API void SetDataType(const FName& InDataType);
-	UE_API void SetIsArray(const bool bInIsArray);
-	UE_API void SetMetaSoundOutput(const FMetaSoundOutput& InMetaSoundOutput);
+	UE_API void SetIsArray(bool bInIsArray);
+	UE_API void SetIsConstructorPin(bool bInIsConstructorPin);
 
-	UFUNCTION()
-	UE_API void OnOutputValueChanged(FName InOutputName, const FMetaSoundOutput& InOutput);
+	UE_API void OnOutputNameChanged(const FName& NewName);
+	UE_API void OnOutputDataTypeChanged(const FName& NewDataType);
+	UE_API void OnOutputIsConstructorPinChanged(const bool bNewIsConstructorPin);
 };
 
 #undef UE_API
diff --git a/Engine/Plugins/Experimental/TechAudioTools/Source/TechAudioToolsMetaSound/Public/ViewModels/MetaSoundViewModelConversionFunctions.h b/Engine/Plugins/Experimental/TechAudioTools/Source/TechAudioToolsMetaSound/Public/ViewModels/MetaSoundViewModelConversionFunctions.h
--- a/Engine/Plugins/Experimental/TechAudioTools/Source/TechAudioToolsMetaSound/Public/ViewModels/MetaSoundViewModelConversionFunctions.h
+++ b/Engine/Plugins/Experimental/TechAudioTools/Source/TechAudioToolsMetaSound/Public/ViewModels/MetaSoundViewModelConversionFunctions.h
@@ -3,12 +3,40 @@
 #pragma once
 
 #include "Kismet/BlueprintFunctionLibrary.h"
+#include "MetasoundDataReferenceMacro.h"
 
 #include "MetaSoundViewModelConversionFunctions.generated.h"
 
+struct FMetasoundFrontendLiteral;
 class UMetaSoundInputViewModel;
 class UMetaSoundOutputViewModel;
 
+namespace TechAudioTools::MetaSound
+{
+	// Returns the given data type with the array type specifier added or removed to match bIsArray.
+	inline FName GetAdjustedDataType(const FName& DataType, const bool bIsArray)
+	{
+		const FString DataTypeString = DataType.ToString();
+		if (bIsArray)
+		{
+			// Append the array type specifier if it isn't already there
+			if (!DataTypeString.EndsWith(METASOUND_DATA_TYPE_NAME_ARRAY_TYPE_SPECIFIER))
+			{
+				return FName(DataTypeString + METASOUND_DATA_TYPE_NAME_ARRAY_TYPE_SPECIFIER);
+			}
+			return DataType;
+		}
+		if (DataTypeString.EndsWith(METASOUND_DATA_TYPE_NAME_ARRAY_TYPE_SPECIFIER))
+		{
+			// Remove the array type specifier
+			FString BaseDataType;
+			DataTypeString.Split(TEXT(":"), &BaseDataType, nullptr);
+			return FName(BaseDataType);
+		}
+		return DataType;
+	}
+} // TechAudioTools::MetaSound
+
 /**
  * Collection of conversion functions to use with MetaSound Viewmodels.
  */
@@ -19,10 +47,26 @@ class UMetaSoundViewModelConversionFunctions : public UBlueprintFunctionLibrary
 
 public:
 	// Searches the given array of MetaSoundInputViewModels for the specified input. Returns nullptr if unable to find a matching viewmodel.
-	UFUNCTION(BlueprintCallable, BlueprintPure, DisplayName = "Find Input Viewmodel by Name", Category = "TechAudioTools")
-	static UMetaSoundInputViewModel* FindInputViewModelByName(const TArray<UMetaSoundInputViewModel*>& MetaSoundInputViewModels, const FName InputName);
+	UFUNCTION(BlueprintCallable, BlueprintPure, DisplayName = "Find MetaSound Input Viewmodel by Name", Category = "TechAudioTools")
+	static TECHAUDIOTOOLSMETASOUND_API UMetaSoundInputViewModel* FindInputViewModelByName(const TArray<UMetaSoundInputViewModel*>& MetaSoundInputViewModels, const FName InputName);
 
 	// Searches the given array of MetaSoundOutputViewModels for the specified output. Returns nullptr if unable to find a matching viewmodel.
-	UFUNCTION(BlueprintCallable, BlueprintPure, DisplayName = "Find Output Viewmodel by Name", Category = "TechAudioTools")
-	static UMetaSoundOutputViewModel* FindOutputViewModelByName(const TArray<UMetaSoundOutputViewModel*>& MetaSoundOutputViewModels, const FName OutputName);
+	UFUNCTION(BlueprintCallable, BlueprintPure, DisplayName = "Find MetaSound Output Viewmodel by Name", Category = "TechAudioTools")
+	static TECHAUDIOTOOLSMETASOUND_API UMetaSoundOutputViewModel* FindOutputViewModelByName(const TArray<UMetaSoundOutputViewModel*>& MetaSoundOutputViewModels, const FName OutputName);
+
+	// Returns the value of the given MetaSound Literal as a text value.
+	UFUNCTION(BlueprintCallable, BlueprintPure, DisplayName = "Get MetaSound Literal Value as Text", Category = "TechAudioTools")
+	static TECHAUDIOTOOLSMETASOUND_API FText GetLiteralValueAsText(const FMetasoundFrontendLiteral& Literal);
+
+	// Checks whether the given member name belongs to a registered MetaSound interface.
+	UFUNCTION(BlueprintCallable, BlueprintPure, DisplayName = "Is MetaSound Interface Member", Category = "TechAudioTools")
+	static TECHAUDIOTOOLSMETASOUND_API bool IsInterfaceMember(const FName& MemberName, UPARAM(DisplayName = "NOT") bool bInvert = false);
+
+	// Returns true if the given MetaSound data type is registered and can be used as an array type.
+	UFUNCTION(BlueprintCallable, BlueprintPure, DisplayName = "Is MetaSound Array Type", Category = "TechAudioTools")
+	static TECHAUDIOTOOLSMETASOUND_API bool IsArrayType(const FName& DataType, UPARAM(DisplayName = "NOT") bool bInvert = false);
+
+	// Returns true if the given MetaSound data type is registered and can be used for constructor pins.
+	UFUNCTION(BlueprintCallable, BlueprintPure, DisplayName = "Is MetaSound Constructor Type", Category = "TechAudioTools")
+	static TECHAUDIOTOOLSMETASOUND_API bool IsConstructorType(const FName& DataType, UPARAM(DisplayName = "NOT") bool bInvert = false);
 };
diff --git a/Engine/Plugins/Experimental/TechAudioTools/Source/TechAudioToolsMetaSound/TechAudioToolsMetaSound.Build.cs b/Engine/Plugins/Experimental/TechAudioTools/Source/TechAudioToolsMetaSound/TechAudioToolsMetaSound.Build.cs
--- a/Engine/Plugins/Experimental/TechAudioTools/Source/TechAudioToolsMetaSound/TechAudioToolsMetaSound.Build.cs
+++ b/Engine/Plugins/Experimental/TechAudioTools/Source/TechAudioToolsMetaSound/TechAudioToolsMetaSound.Build.cs
@@ -15,14 +15,19 @@ public class TechAudioToolsMetaSound : ModuleRules
 				"MetasoundEngine",
 				"MetasoundFrontend",
 				"ModelViewViewModel",
+				"SlateCore",
+				"TechAudioTools",
             }
         );
 
         PrivateDependencyModuleNames.AddRange(
             new string[]
             {
+	            "AudioExtensions",
                 "CoreUObject",
                 "Engine",
+                "Slate",
+                "UMG",
             }
         );
     }
diff --git a/Engine/Plugins/Experimental/TechAudioTools/Source/TechAudioToolsMetaSoundEditor/Private/ViewModels/MetaSoundEditorViewModel.cpp b/Engine/Plugins/Experimental/TechAudioTools/Source/TechAudioToolsMetaSoundEditor/Private/ViewModels/MetaSoundEditorViewModel.cpp
--- a/Engine/Plugins/Experimental/TechAudioTools/Source/TechAudioToolsMetaSoundEditor/Private/ViewModels/MetaSoundEditorViewModel.cpp
+++ b/Engine/Plugins/Experimental/TechAudioTools/Source/TechAudioToolsMetaSoundEditor/Private/ViewModels/MetaSoundEditorViewModel.cpp
@@ -2,34 +2,19 @@
 
 #include "ViewModels/MetaSoundEditorViewModel.h"
 
+#include "MetasoundBuilderSubsystem.h"
 #include "Editor/EditorEngine.h"
 #include "Logging/StructuredLog.h"
 #include "MetasoundDocumentBuilderRegistry.h"
 #include "MetasoundEditorBuilderListener.h"
 #include "MetasoundEditorSubsystem.h"
 #include "MetasoundOutputSubsystem.h"
+#include "TechAudioToolsLog.h"
 
 #include UE_INLINE_GENERATED_CPP_BY_NAME(MetaSoundEditorViewModel)
 
 extern UNREALED_API UEditorEngine* GEditor;
 
-DEFINE_LOG_CATEGORY(LogTechAudioToolsMetaSoundEditor);
-
-bool UMetaSoundEditorViewModel::IsEditorOnly() const
-{
-	return true;
-}
-
-UWorld* UMetaSoundEditorViewModel::GetWorld() const
-{
-	if (HasAllFlags(RF_ClassDefaultObject))
-	{
-		return nullptr;
-	}
-
-	return GEditor ? GEditor->GetEditorWorldContext(false).World() : nullptr;
-}
-
 void UMetaSoundEditorViewModel::InitializeMetaSound(const TScriptInterface<IMetaSoundDocumentInterface> InMetaSound)
 {
 	if (!InMetaSound)
@@ -38,37 +23,30 @@ void UMetaSoundEditorViewModel::InitializeMetaSound(const TScriptInterface<IMeta
 		return;
 	}
 
-	Builder = &Metasound::Engine::FDocumentBuilderRegistry::GetChecked().FindOrBeginBuilding(*InMetaSound.GetObject());
+	Builder = &Metasound::Engine::FDocumentBuilderRegistry::GetChecked().FindOrBeginBuilding(*InMetaSound.GetObject(), true);
 	Initialize(Builder);
 }
 
 void UMetaSoundEditorViewModel::Initialize(UMetaSoundBuilderBase* InBuilder)
 {
 	Super::Initialize(InBuilder);
-	BindDelegates();
-}
 
-void UMetaSoundEditorViewModel::Reset()
-{
-	Super::Reset();
-	UnbindDelegates();
-}
+	const FMetaSoundFrontendDocumentBuilder& DocumentBuilder = InBuilder->GetConstBuilder();
+	const FMetasoundFrontendDocument& FrontendDocument = DocumentBuilder.GetConstDocumentChecked();
 
-void UMetaSoundEditorViewModel::InitializeProperties(const FMetasoundFrontendDocument& FrontendDocument)
-{
-	Super::InitializeProperties(FrontendDocument);
+	OnDisplayNameChanged(FrontendDocument.RootGraph.Metadata.GetDisplayName());
+	OnDescriptionChanged(FrontendDocument.RootGraph.Metadata.GetDescription());
+	OnAuthorChanged(FrontendDocument.RootGraph.Metadata.GetAuthor());
+	OnKeywordsChanged(FrontendDocument.RootGraph.Metadata.GetKeywords());
+	OnCategoryHierarchyChanged(FrontendDocument.RootGraph.Metadata.GetCategoryHierarchy());
+	OnIsDeprecatedChanged(FrontendDocument.RootGraph.Metadata.GetIsDeprecated());
 
-	UE_MVVM_SET_PROPERTY_VALUE(DisplayName, FrontendDocument.RootGraph.Metadata.GetDisplayName());
-	UE_MVVM_SET_PROPERTY_VALUE(Description, FrontendDocument.RootGraph.Metadata.GetDescription());
-	UE_MVVM_SET_PROPERTY_VALUE(Author, FrontendDocument.RootGraph.Metadata.GetAuthor());
-	UE_MVVM_SET_PROPERTY_VALUE(Keywords, FrontendDocument.RootGraph.Metadata.GetKeywords());
-	UE_MVVM_SET_PROPERTY_VALUE(CategoryHierarchy, FrontendDocument.RootGraph.Metadata.GetCategoryHierarchy());
-	UE_MVVM_SET_PROPERTY_VALUE(bIsDeprecated, FrontendDocument.RootGraph.Metadata.GetIsDeprecated());
+	BindDelegates();
 }
 
-void UMetaSoundEditorViewModel::ResetProperties()
+void UMetaSoundEditorViewModel::Reset()
 {
-	Super::ResetProperties();
+	Super::Reset();
 
 	UE_MVVM_SET_PROPERTY_VALUE(DisplayName, FText());
 	UE_MVVM_SET_PROPERTY_VALUE(Description, FText());
@@ -76,6 +54,8 @@ void UMetaSoundEditorViewModel::ResetProperties()
 	UE_MVVM_SET_PROPERTY_VALUE(Keywords, TArray<FText>());
 	UE_MVVM_SET_PROPERTY_VALUE(CategoryHierarchy, TArray<FText>());
 	UE_MVVM_SET_PROPERTY_VALUE(bIsDeprecated, false);
+
+	UnbindDelegates();
 }
 
 void UMetaSoundEditorViewModel::SetMetaSoundDisplayName(const FText& InDisplayName)
@@ -132,337 +112,338 @@ void UMetaSoundEditorViewModel::SetIsDeprecated(const bool bInIsDeprecated)
 	}
 }
 
-void UMetaSoundEditorViewModel::CreateInputViewModel(const FMetasoundFrontendClassInput& InInput)
+void UMetaSoundEditorViewModel::SetInputDisplayName(const FName& InputName, const FText& InDisplayName) const
 {
-	Super::CreateInputViewModel(InInput);
-
-	UMetaSoundInputViewModel* InputViewModel = *InputViewModels.Find(InInput.Name);
-	if (UMetaSoundInputEditorViewModel* InputEditorViewModel = Cast<UMetaSoundInputEditorViewModel>(InputViewModel))
+	if (Builder)
 	{
-		if (BuilderListener && InputEditorViewModel)
-		{
-			BuilderListener->OnGraphInputDefaultChangedDelegate.AddDynamic(InputEditorViewModel, &UMetaSoundInputEditorViewModel::OnInputDefaultChanged);
-		}
+		FMetaSoundFrontendDocumentBuilder& DocumentBuilder = Builder->GetBuilder();
+		DocumentBuilder.SetGraphInputDisplayName(InputName, InDisplayName);
 	}
 }
 
-UMetaSoundInputViewModel* UMetaSoundEditorViewModel::CreateInputViewModelInstance()
+void UMetaSoundEditorViewModel::SetInputDescription(const FName& InputName, const FText& InDescription) const
 {
-    return NewObject<UMetaSoundInputEditorViewModel>(this, UMetaSoundInputEditorViewModel::StaticClass());
+	if (Builder)
+	{
+		FMetaSoundFrontendDocumentBuilder& DocumentBuilder = Builder->GetBuilder();
+		DocumentBuilder.SetGraphInputDescription(InputName, InDescription);
+	}
 }
 
-UMetaSoundOutputViewModel* UMetaSoundEditorViewModel::CreateOutputViewModelInstance()
+void UMetaSoundEditorViewModel::SetInputSortOrderIndex(const FName& InputName, const int32 InSortOrderIndex) const
 {
-    return NewObject<UMetaSoundOutputEditorViewModel>(this, UMetaSoundOutputEditorViewModel::StaticClass());
+	if (Builder)
+	{
+		FMetaSoundFrontendDocumentBuilder& DocumentBuilder = Builder->GetBuilder();
+		DocumentBuilder.SetGraphInputSortOrderIndex(InputName, InSortOrderIndex);
+	}
 }
 
-void UMetaSoundEditorViewModel::BindDelegates()
+void UMetaSoundEditorViewModel::SetInputIsAdvancedDisplay(const FName& InputName, const bool bInIsAdvancedDisplay) const
 {
-	if (!Builder)
+	if (Builder)
 	{
-		UE_LOGFMT(LogTechAudioToolsMetaSoundEditor, Log, "Could not bind MetaSoundViewModel delegates. Builder was null.");
-		SetIsInitialized(false);
-		return;
+		FMetaSoundFrontendDocumentBuilder& DocumentBuilder = Builder->GetBuilder();
+		DocumentBuilder.SetGraphInputAdvancedDisplay(InputName, bInIsAdvancedDisplay);
 	}
+}
 
-	UMetaSoundEditorSubsystem* EditorSubsystem = GEditor->GetEditorSubsystem<UMetaSoundEditorSubsystem>();
-	if (!EditorSubsystem)
+void UMetaSoundEditorViewModel::SetOutputDisplayName(const FName& OutputName, const FText& InDisplayName) const
+{
+	if (Builder)
 	{
-		UE_LOGFMT(LogTechAudioToolsMetaSoundEditor, Log, "Could not bind MetaSoundViewModel delegates. Unable to locate MetaSound Editor Subsystem.");
-		SetIsInitialized(false);
-		return;
+		FMetaSoundFrontendDocumentBuilder& DocumentBuilder = Builder->GetBuilder();
+		DocumentBuilder.SetGraphOutputDisplayName(OutputName, InDisplayName);
 	}
+}
 
-	EMetaSoundBuilderResult Result;
-	BuilderListener = EditorSubsystem->AddBuilderDelegateListener(Builder, Result);
-
-	if (Result != EMetaSoundBuilderResult::Succeeded)
+void UMetaSoundEditorViewModel::SetOutputDescription(const FName& OutputName, const FText& InDescription) const
+{
+	if (Builder)
 	{
-		UE_LOGFMT(LogTechAudioToolsMetaSoundEditor, Log, "Could not bind MetaSoundViewModel delegates. Failed to create BuilderListener.");
-		SetIsInitialized(false);
-		return;
+		FMetaSoundFrontendDocumentBuilder& DocumentBuilder = Builder->GetBuilder();
+		DocumentBuilder.SetGraphOutputDescription(OutputName, InDescription);
 	}
+}
 
-	if (BuilderListener)
+void UMetaSoundEditorViewModel::SetOutputSortOrderIndex(const FName& OutputName, const int32 InSortOrderIndex) const
+{
+	if (Builder)
 	{
-		BuilderListener->OnGraphInputAddedDelegate.AddDynamic(this, &ThisClass::OnInputAdded);
-		BuilderListener->OnRemovingGraphInputDelegate.AddDynamic(this, &ThisClass::OnInputRemoved);
-		BuilderListener->OnGraphInputNameChangedDelegate.AddDynamic(this, &ThisClass::OnInputNameChanged);
-		BuilderListener->OnGraphInputDataTypeChangedDelegate.AddDynamic(this, &ThisClass::OnInputDataTypeChanged);
-
-		BuilderListener->OnGraphOutputAddedDelegate.AddDynamic(this, &ThisClass::OnOutputAdded);
-		BuilderListener->OnRemovingGraphOutputDelegate.AddDynamic(this, &ThisClass::OnOutputRemoved);
-		BuilderListener->OnGraphOutputNameChangedDelegate.AddDynamic(this, &ThisClass::OnOutputNameChanged);
-		BuilderListener->OnGraphOutputDataTypeChangedDelegate.AddDynamic(this, &ThisClass::OnOutputDataTypeChanged);
-
-		for (const TPair<FName, UMetaSoundInputViewModel*> InputViewModel : InputViewModels)
-		{
-			if (UMetaSoundInputEditorViewModel* InputEditorViewModel = Cast<UMetaSoundInputEditorViewModel>(InputViewModel.Value))
-			{
-				BuilderListener->OnGraphInputDefaultChangedDelegate.AddDynamic(InputEditorViewModel, &UMetaSoundInputEditorViewModel::OnInputDefaultChanged);
-			}
-		}
+		FMetaSoundFrontendDocumentBuilder& DocumentBuilder = Builder->GetBuilder();
+		DocumentBuilder.SetGraphOutputSortOrderIndex(OutputName, InSortOrderIndex);
 	}
 }
 
-void UMetaSoundEditorViewModel::UnbindDelegates()
+void UMetaSoundEditorViewModel::SetOutputIsAdvancedDisplay(const FName& OutputName, const bool bInIsAdvancedDisplay) const
 {
-	if (BuilderListener)
+	if (Builder)
 	{
-		BuilderListener->RemoveAllDelegates();
-		BuilderListener = nullptr;
+		FMetaSoundFrontendDocumentBuilder& DocumentBuilder = Builder->GetBuilder();
+		DocumentBuilder.SetGraphOutputAdvancedDisplay(OutputName, bInIsAdvancedDisplay);
 	}
 }
 
-void UMetaSoundEditorViewModel::OnInputAdded(const FName VertexName, const FName DataType)
+void UMetaSoundEditorViewModel::OnInputDisplayNameChanged(const FName InputName, const FText InDisplayName)
 {
-	const FMetaSoundFrontendDocumentBuilder& DocumentBuilder = Builder->GetConstBuilder();
-	const FMetasoundFrontendDocument& FrontendDocument = DocumentBuilder.GetConstDocumentChecked();
-	for (const FMetasoundFrontendClassInput& Input : FrontendDocument.RootGraph.GetDefaultInterface().Inputs)
+	if (UMetaSoundInputEditorViewModel* InputEditorViewModel = Cast<UMetaSoundInputEditorViewModel>(*InputViewModels.Find(InputName)))
 	{
-		if (Input.Name == VertexName)
-		{
-			CreateInputViewModel(Input);
-			return;
-		}
+		InputEditorViewModel->OnInputDisplayNameChanged(InDisplayName);
 	}
 }
 
-void UMetaSoundEditorViewModel::OnInputRemoved(const FName VertexName, const FName DataType)
+void UMetaSoundEditorViewModel::OnInputDescriptionChanged(const FName InputName, FText InDescription)
 {
-	InputViewModels.Remove(VertexName);
-	UE_MVVM_BROADCAST_FIELD_VALUE_CHANGED(GetInputViewModels);
+	if (UMetaSoundInputEditorViewModel* InputEditorViewModel = Cast<UMetaSoundInputEditorViewModel>(*InputViewModels.Find(InputName)))
+	{
+		InputEditorViewModel->OnInputDescriptionChanged(InDescription);
+	}
 }
 
-void UMetaSoundEditorViewModel::OnInputNameChanged(const FName OldName, const FName NewName)
+void UMetaSoundEditorViewModel::OnInputSortOrderIndexChanged(const FName InputName, const int32 InSortOrderIndex)
 {
-	if (UMetaSoundInputViewModel* InputViewModel = *InputViewModels.Find(OldName))
+	if (UMetaSoundInputEditorViewModel* InputEditorViewModel = Cast<UMetaSoundInputEditorViewModel>(*InputViewModels.Find(InputName)))
 	{
-		InputViewModel->SetInputName(NewName);
-		UE_MVVM_BROADCAST_FIELD_VALUE_CHANGED(GetInputViewModels);
+		InputEditorViewModel->OnInputSortOrderIndexChanged(InSortOrderIndex);
 	}
 }
 
-void UMetaSoundEditorViewModel::OnInputDataTypeChanged(const FName VertexName, const FName DataType)
+void UMetaSoundEditorViewModel::OnInputIsAdvancedDisplayChanged(const FName InputName, const bool bInIsAdvancedDisplay)
 {
-	if (UMetaSoundInputViewModel* InputViewModel = *InputViewModels.Find(VertexName))
+	if (UMetaSoundInputEditorViewModel* InputEditorViewModel = Cast<UMetaSoundInputEditorViewModel>(*InputViewModels.Find(InputName)))
 	{
-		InputViewModel->SetDataType(DataType);
-		UE_MVVM_BROADCAST_FIELD_VALUE_CHANGED(GetInputViewModels);
+		InputEditorViewModel->OnInputIsAdvancedDisplayChanged(bInIsAdvancedDisplay);
 	}
 }
 
-void UMetaSoundEditorViewModel::OnOutputAdded(const FName VertexName, const FName DataType)
+void UMetaSoundEditorViewModel::OnOutputDisplayNameChanged(const FName OutputName, const FText NewName)
 {
-	if (Builder)
+	if (UMetaSoundOutputEditorViewModel* OutputEditorViewModel = Cast<UMetaSoundOutputEditorViewModel>(OutputViewModels.FindAndRemoveChecked(OutputName)))
 	{
-		const FMetaSoundFrontendDocumentBuilder& DocumentBuilder = Builder->GetConstBuilder();
-		const FMetasoundFrontendDocument& FrontendDocument = DocumentBuilder.GetConstDocumentChecked();
-		for (const FMetasoundFrontendClassOutput& Output : FrontendDocument.RootGraph.GetDefaultInterface().Outputs)
-		{
-			if (Output.Name == VertexName)
-			{
-				CreateOutputViewModel(Output);
-				return;
-			}
-		}
+		OutputEditorViewModel->OnOutputDisplayNameChanged(NewName);
 	}
 }
 
-void UMetaSoundEditorViewModel::OnOutputRemoved(const FName VertexName, const FName DataType)
+void UMetaSoundEditorViewModel::OnOutputDescriptionChanged(const FName OutputName, const FText InDescription)
 {
-	OutputViewModels.Remove(VertexName);
-	UE_MVVM_BROADCAST_FIELD_VALUE_CHANGED(GetOutputViewModels);
+	if (UMetaSoundOutputEditorViewModel* OutputEditorViewModel = Cast<UMetaSoundOutputEditorViewModel>(OutputViewModels.FindAndRemoveChecked(OutputName)))
+	{
+		OutputEditorViewModel->OnOutputDescriptionChanged(InDescription);
+	}
 }
 
-void UMetaSoundEditorViewModel::OnOutputNameChanged(const FName OldName, const FName NewName)
+void UMetaSoundEditorViewModel::OnOutputSortOrderIndexChanged(const FName OutputName, const int32 InSortOrderIndex)
 {
-	if (UMetaSoundOutputViewModel* OutputViewModel = *OutputViewModels.Find(OldName))
+	if (UMetaSoundOutputEditorViewModel* OutputEditorViewModel = Cast<UMetaSoundOutputEditorViewModel>(OutputViewModels.FindAndRemoveChecked(OutputName)))
 	{
-		OutputViewModel->SetOutputName(NewName);
-		UE_MVVM_BROADCAST_FIELD_VALUE_CHANGED(GetOutputViewModels);
+		OutputEditorViewModel->OnOutputSortOrderIndexChanged(InSortOrderIndex);
 	}
 }
 
-void UMetaSoundEditorViewModel::OnOutputDataTypeChanged(const FName VertexName, const FName DataType)
+void UMetaSoundEditorViewModel::OnOutputIsAdvancedDisplayChanged(const FName OutputName, const bool bInIsAdvancedDisplay)
 {
-	if (UMetaSoundOutputViewModel* OutputViewModel = *OutputViewModels.Find(VertexName))
+	if (UMetaSoundOutputEditorViewModel* OutputEditorViewModel = Cast<UMetaSoundOutputEditorViewModel>(OutputViewModels.FindAndRemoveChecked(OutputName)))
 	{
-		OutputViewModel->SetDataType(DataType);
-		UE_MVVM_BROADCAST_FIELD_VALUE_CHANGED(GetOutputViewModels);
+		OutputEditorViewModel->OnOutputIsAdvancedDisplayChanged(bInIsAdvancedDisplay);
 	}
 }
 
-bool UMetaSoundInputEditorViewModel::IsEditorOnly() const
+UMetaSoundInputViewModel* UMetaSoundEditorViewModel::CreateInputViewModel(const FMetasoundFrontendClassInput& InInput)
 {
-	return true;
+	if (UMetaSoundInputEditorViewModel* InputEditorViewModel = Cast<UMetaSoundInputEditorViewModel>(Super::CreateInputViewModel(InInput)))
+	{
+		InputEditorViewModel->OnInputDisplayNameChanged(InInput.Metadata.GetDisplayName());
+		InputEditorViewModel->OnInputDescriptionChanged(InInput.Metadata.GetDescription());
+		InputEditorViewModel->OnInputSortOrderIndexChanged(InInput.Metadata.SortOrderIndex);
+		InputEditorViewModel->OnInputIsAdvancedDisplayChanged(InInput.Metadata.bIsAdvancedDisplay);
+
+		return InputEditorViewModel;
+	}
+
+	return nullptr;
 }
 
-UWorld* UMetaSoundInputEditorViewModel::GetWorld() const
+UMetaSoundOutputViewModel* UMetaSoundEditorViewModel::CreateOutputViewModel(const FMetasoundFrontendClassOutput& InOutput)
 {
-	if (HasAllFlags(RF_ClassDefaultObject))
+	if (UMetaSoundOutputEditorViewModel* OutputEditorViewModel = Cast<UMetaSoundOutputEditorViewModel>(Super::CreateOutputViewModel(InOutput)))
 	{
-		return nullptr;
+		OutputEditorViewModel->OnOutputDisplayNameChanged(InOutput.Metadata.GetDisplayName());
+		OutputEditorViewModel->OnOutputDescriptionChanged(InOutput.Metadata.GetDescription());
+		OutputEditorViewModel->OnOutputSortOrderIndexChanged(InOutput.Metadata.SortOrderIndex);
+		OutputEditorViewModel->OnOutputIsAdvancedDisplayChanged(InOutput.Metadata.bIsAdvancedDisplay);
+
+		return OutputEditorViewModel;
 	}
 
-	return GEditor ? GEditor->GetEditorWorldContext(false).World() : nullptr;
+	return nullptr;
 }
 
-void UMetaSoundInputEditorViewModel::InitializeInput(UMetaSoundBuilderBase* InBuilder, const FMetasoundFrontendClassInput& InInput)
+TSubclassOf<UMetaSoundInputViewModel> UMetaSoundEditorViewModel::GetInputViewModelClass() const
 {
-	Super::InitializeInput(InBuilder, InInput);
-
-	UE_MVVM_SET_PROPERTY_VALUE(InputDisplayName, InInput.Metadata.GetDisplayName());
-	UE_MVVM_SET_PROPERTY_VALUE(InputDescription, InInput.Metadata.GetDescription());
-	UE_MVVM_SET_PROPERTY_VALUE(SortOrderIndex, InInput.Metadata.SortOrderIndex);
-	UE_MVVM_SET_PROPERTY_VALUE(bIsAdvancedDisplay, InInput.Metadata.bIsAdvancedDisplay);
+	return UMetaSoundInputEditorViewModel::StaticClass();
 }
 
-void UMetaSoundInputEditorViewModel::ResetInput()
+TSubclassOf<UMetaSoundOutputViewModel> UMetaSoundEditorViewModel::GetOutputViewModelClass() const
 {
-	Super::ResetInput();
-
-	UE_MVVM_SET_PROPERTY_VALUE(InputDisplayName, FText());
-	UE_MVVM_SET_PROPERTY_VALUE(InputDescription, FText());
-	UE_MVVM_SET_PROPERTY_VALUE(SortOrderIndex, 0);
-	UE_MVVM_SET_PROPERTY_VALUE(bIsAdvancedDisplay, false);
+	return UMetaSoundOutputEditorViewModel::StaticClass();
 }
 
-void UMetaSoundInputEditorViewModel::SetInputDisplayName(const FText& InDisplayName)
+void UMetaSoundEditorViewModel::BindDelegates()
 {
-	if (Builder)
+	if (!Builder)
 	{
-		FMetaSoundFrontendDocumentBuilder& DocumentBuilder = Builder->GetBuilder();
-		if (DocumentBuilder.IsValid() && UE_MVVM_SET_PROPERTY_VALUE(InputDisplayName, InDisplayName))
-		{
-			DocumentBuilder.SetGraphInputDisplayName(InputName, InDisplayName);
-		}
+		UE_LOGFMT(LogTechAudioTools, Log, "Could not bind MetaSoundViewModel delegates. Builder was null.");
+		SetIsInitialized(false);
+		return;
+	}
+
+	UMetaSoundEditorSubsystem* EditorSubsystem = GEditor->GetEditorSubsystem<UMetaSoundEditorSubsystem>();
+	if (!EditorSubsystem)
+	{
+		UE_LOGFMT(LogTechAudioTools, Log, "Could not bind MetaSoundViewModel delegates. Unable to locate MetaSound Editor Subsystem.");
+		SetIsInitialized(false);
+		return;
+	}
+
+	EMetaSoundBuilderResult Result;
+	BuilderListener = EditorSubsystem->AddBuilderDelegateListener(Builder, Result);
+
+	if (Result != EMetaSoundBuilderResult::Succeeded)
+	{
+		UE_LOGFMT(LogTechAudioTools, Warning, "Could not bind MetaSoundViewModel delegates. Failed to create BuilderListener.");
+		SetIsInitialized(false);
+		return;
+	}
+
+	if (BuilderListener)
+	{
+		BuilderListener->OnGraphInputAddedDelegate.AddDynamic(this, &UMetaSoundViewModel::OnInputAdded);
+		BuilderListener->OnRemovingGraphInputDelegate.AddDynamic(this, &UMetaSoundViewModel::OnInputRemoved);
+		BuilderListener->OnGraphInputNameChangedDelegate.AddDynamic(this, &UMetaSoundViewModel::OnInputNameChanged);
+		BuilderListener->OnGraphInputDataTypeChangedDelegate.AddDynamic(this, &UMetaSoundViewModel::OnInputDataTypeChanged);
+		BuilderListener->OnGraphInputDefaultChangedDelegate.AddDynamic(this, &UMetaSoundViewModel::OnInputDefaultChanged);
+		BuilderListener->OnGraphInputInheritsDefaultChangedDelegate.AddDynamic(this, &UMetaSoundViewModel::OnInputInheritsDefaultChanged);
+		BuilderListener->OnGraphInputIsConstructorPinChangedDelegate.AddDynamic(this, &UMetaSoundViewModel::OnInputIsConstructorPinChanged);
+
+		BuilderListener->OnGraphOutputAddedDelegate.AddDynamic(this, &UMetaSoundViewModel::OnOutputAdded);
+		BuilderListener->OnRemovingGraphOutputDelegate.AddDynamic(this, &UMetaSoundViewModel::OnOutputRemoved);
+		BuilderListener->OnGraphOutputNameChangedDelegate.AddDynamic(this, &UMetaSoundViewModel::OnOutputNameChanged);
+		BuilderListener->OnGraphOutputDataTypeChangedDelegate.AddDynamic(this, &UMetaSoundViewModel::OnOutputDataTypeChanged);
+		BuilderListener->OnGraphOutputIsConstructorPinChangedDelegate.AddDynamic(this, &UMetaSoundViewModel::OnOutputIsConstructorPinChanged);
+
+		// Note: Document metadata delegates currently share a single delegate in the builder listener which is called when any metadata changes, which calls all of these together.
+		BuilderListener->OnDocumentDisplayNameChangedDelegate.AddDynamic(this, &ThisClass::OnDisplayNameChanged);
+		BuilderListener->OnDocumentDescriptionChangedDelegate.AddDynamic(this, &ThisClass::OnDescriptionChanged);
+		BuilderListener->OnDocumentAuthorChangedDelegate.AddDynamic(this, &ThisClass::OnAuthorChanged);
+		BuilderListener->OnDocumentKeywordsChangedDelegate.AddDynamic(this, &ThisClass::OnKeywordsChanged);
+		BuilderListener->OnDocumentCategoryHierarchyChangedDelegate.AddDynamic(this, &ThisClass::OnCategoryHierarchyChanged);
+		BuilderListener->OnDocumentIsDeprecatedChangedDelegate.AddDynamic(this, &ThisClass::OnIsDeprecatedChanged);
+
+		BuilderListener->OnGraphInputDisplayNameChangedDelegate.AddDynamic(this, &ThisClass::OnInputDisplayNameChanged);
+		BuilderListener->OnGraphInputDescriptionChangedDelegate.AddDynamic(this, &ThisClass::OnInputDescriptionChanged);
+		BuilderListener->OnGraphInputSortOrderIndexChangedDelegate.AddDynamic(this, &ThisClass::OnInputSortOrderIndexChanged);
+		BuilderListener->OnGraphInputIsAdvancedDisplayChangedDelegate.AddDynamic(this, &ThisClass::OnInputIsAdvancedDisplayChanged);
+
+		BuilderListener->OnGraphOutputDisplayNameChangedDelegate.AddDynamic(this, &ThisClass::OnOutputDisplayNameChanged);
+		BuilderListener->OnGraphOutputDescriptionChangedDelegate.AddDynamic(this, &ThisClass::OnOutputDescriptionChanged);
+		BuilderListener->OnGraphOutputSortOrderIndexChangedDelegate.AddDynamic(this, &ThisClass::OnOutputSortOrderIndexChanged);
+		BuilderListener->OnGraphOutputIsAdvancedDisplayChangedDelegate.AddDynamic(this, &ThisClass::OnOutputIsAdvancedDisplayChanged);
 	}
 }
 
-void UMetaSoundInputEditorViewModel::SetInputDescription(const FText& InDescription)
+void UMetaSoundEditorViewModel::UnbindDelegates()
 {
-	if (Builder)
+	if (BuilderListener)
 	{
-		FMetaSoundFrontendDocumentBuilder& DocumentBuilder = Builder->GetBuilder();
-		if (DocumentBuilder.IsValid() && UE_MVVM_SET_PROPERTY_VALUE(InputDescription, InDescription))
-		{
-			DocumentBuilder.SetGraphInputDescription(InputName, InDescription);
-		}
+		BuilderListener->RemoveAllDelegates();
+		BuilderListener = nullptr;
 	}
 }
 
-void UMetaSoundInputEditorViewModel::SetSortOrderIndex(const int32 InSortOrderIndex)
+void UMetaSoundInputEditorViewModel::SetInputDisplayName(const FText& InDisplayName)
 {
-	if (Builder)
+	if (UE_MVVM_SET_PROPERTY_VALUE(InputDisplayName, InDisplayName))
 	{
-		FMetaSoundFrontendDocumentBuilder& DocumentBuilder = Builder->GetBuilder();
-		if (DocumentBuilder.IsValid() && UE_MVVM_SET_PROPERTY_VALUE(SortOrderIndex, InSortOrderIndex))
+		if (const UMetaSoundEditorViewModel* MetaSoundEditorViewModel = Cast<UMetaSoundEditorViewModel>(GetOuter()))
 		{
-			DocumentBuilder.SetGraphInputSortOrderIndex(InputName, InSortOrderIndex);
+			MetaSoundEditorViewModel->SetInputDisplayName(InputName, InDisplayName);
 		}
 	}
 }
 
-void UMetaSoundInputEditorViewModel::SetIsAdvancedDisplay(const bool bInIsAdvancedDisplay)
+void UMetaSoundInputEditorViewModel::SetInputDescription(const FText& InDescription)
 {
-	if (Builder)
+	if (UE_MVVM_SET_PROPERTY_VALUE(InputDescription, InDescription))
 	{
-		FMetaSoundFrontendDocumentBuilder& DocumentBuilder = Builder->GetBuilder();
-		if (DocumentBuilder.IsValid() && UE_MVVM_SET_PROPERTY_VALUE(bIsAdvancedDisplay, bInIsAdvancedDisplay))
+		if (const UMetaSoundEditorViewModel* MetaSoundEditorViewModel = Cast<UMetaSoundEditorViewModel>(GetOuter()))
 		{
-			DocumentBuilder.SetGraphInputAdvancedDisplay(InputName, bInIsAdvancedDisplay);
+			MetaSoundEditorViewModel->SetInputDescription(InputName, InDescription);
 		}
 	}
 }
 
-void UMetaSoundInputEditorViewModel::OnInputDefaultChanged(const FName VertexName, FMetasoundFrontendLiteral LiteralValue, FName PageName)
+void UMetaSoundInputEditorViewModel::SetSortOrderIndex(const int32 InSortOrderIndex)
 {
-	if (InputName == VertexName)
+	if (UE_MVVM_SET_PROPERTY_VALUE(SortOrderIndex, InSortOrderIndex))
 	{
-		SetLiteral(LiteralValue);
+		if (const UMetaSoundEditorViewModel* MetaSoundEditorViewModel = Cast<UMetaSoundEditorViewModel>(GetOuter()))
+		{
+			MetaSoundEditorViewModel->SetInputSortOrderIndex(InputName, InSortOrderIndex);
+		}
 	}
 }
 
-bool UMetaSoundOutputEditorViewModel::IsEditorOnly() const
-{
-	return true;
-}
-
-UWorld* UMetaSoundOutputEditorViewModel::GetWorld() const
+void UMetaSoundInputEditorViewModel::SetIsAdvancedDisplay(const bool bInIsAdvancedDisplay)
 {
-	if (HasAllFlags(RF_ClassDefaultObject))
+	if (UE_MVVM_SET_PROPERTY_VALUE(bIsAdvancedDisplay, bInIsAdvancedDisplay))
 	{
-		return nullptr;
+		if (const UMetaSoundEditorViewModel* MetaSoundEditorViewModel = Cast<UMetaSoundEditorViewModel>(GetOuter()))
+		{
+			MetaSoundEditorViewModel->SetInputIsAdvancedDisplay(InputName, bInIsAdvancedDisplay);
+		}
 	}
-
-	return GEditor ? GEditor->GetEditorWorldContext(false).World() : nullptr;
-}
-
-void UMetaSoundOutputEditorViewModel::InitializeOutput(UMetaSoundBuilderBase* InBuilder, const FMetasoundFrontendClassOutput& InOutput)
-{
-	Super::InitializeOutput(InBuilder, InOutput);
-
-	UE_MVVM_SET_PROPERTY_VALUE(OutputDisplayName, InOutput.Metadata.GetDisplayName());
-	UE_MVVM_SET_PROPERTY_VALUE(OutputDescription, InOutput.Metadata.GetDescription());
-	UE_MVVM_SET_PROPERTY_VALUE(SortOrderIndex, InOutput.Metadata.SortOrderIndex);
-	UE_MVVM_SET_PROPERTY_VALUE(bIsAdvancedDisplay, InOutput.Metadata.bIsAdvancedDisplay);
-}
-
-void UMetaSoundOutputEditorViewModel::ResetOutput()
-{
-	Super::ResetOutput();
-
-	UE_MVVM_SET_PROPERTY_VALUE(OutputDisplayName, FText());
-	UE_MVVM_SET_PROPERTY_VALUE(OutputDescription, FText());
-	UE_MVVM_SET_PROPERTY_VALUE(SortOrderIndex, 0);
-	UE_MVVM_SET_PROPERTY_VALUE(bIsAdvancedDisplay, false);
 }
 
 void UMetaSoundOutputEditorViewModel::SetOutputDisplayName(const FText& InDisplayName)
 {
-	if (Builder)
+	if (UE_MVVM_SET_PROPERTY_VALUE(OutputDisplayName, InDisplayName))
 	{
-		FMetaSoundFrontendDocumentBuilder& DocumentBuilder = Builder->GetBuilder();
-		if (DocumentBuilder.IsValid() && UE_MVVM_SET_PROPERTY_VALUE(OutputDisplayName, InDisplayName))
+		if (const UMetaSoundEditorViewModel* MetaSoundEditorViewModel = Cast<UMetaSoundEditorViewModel>(GetOuter()))
 		{
-			DocumentBuilder.SetGraphOutputDisplayName(OutputName, InDisplayName);
+			MetaSoundEditorViewModel->SetOutputDisplayName(OutputName, InDisplayName);
 		}
 	}
 }
 
 void UMetaSoundOutputEditorViewModel::SetOutputDescription(const FText& InDescription)
 {
-	if (Builder)
+	if (UE_MVVM_SET_PROPERTY_VALUE(OutputDescription, InDescription))
 	{
-		FMetaSoundFrontendDocumentBuilder& DocumentBuilder = Builder->GetBuilder();
-		if (DocumentBuilder.IsValid() && UE_MVVM_SET_PROPERTY_VALUE(OutputDescription, InDescription))
+		if (const UMetaSoundEditorViewModel* MetaSoundEditorViewModel = Cast<UMetaSoundEditorViewModel>(GetOuter()))
 		{
-			DocumentBuilder.SetGraphOutputDescription(OutputName, InDescription);
+			MetaSoundEditorViewModel->SetOutputDescription(OutputName, InDescription);
 		}
 	}
 }
 
 void UMetaSoundOutputEditorViewModel::SetSortOrderIndex(const int32 InSortOrderIndex)
 {
-	if (Builder)
+	if (UE_MVVM_SET_PROPERTY_VALUE(SortOrderIndex, InSortOrderIndex))
 	{
-		FMetaSoundFrontendDocumentBuilder& DocumentBuilder = Builder->GetBuilder();
-		if (DocumentBuilder.IsValid() && UE_MVVM_SET_PROPERTY_VALUE(SortOrderIndex, InSortOrderIndex))
+		if (const UMetaSoundEditorViewModel* MetaSoundEditorViewModel = Cast<UMetaSoundEditorViewModel>(GetOuter()))
 		{
-			DocumentBuilder.SetGraphOutputSortOrderIndex(OutputName, InSortOrderIndex);
+			MetaSoundEditorViewModel->SetOutputSortOrderIndex(OutputName, InSortOrderIndex);
 		}
 	}
 }
 
 void UMetaSoundOutputEditorViewModel::SetIsAdvancedDisplay(const bool bInIsAdvancedDisplay)
 {
-	if (Builder)
+	if (UE_MVVM_SET_PROPERTY_VALUE(bIsAdvancedDisplay, bInIsAdvancedDisplay))
 	{
-		FMetaSoundFrontendDocumentBuilder& DocumentBuilder = Builder->GetBuilder();
-		if (DocumentBuilder.IsValid() && UE_MVVM_SET_PROPERTY_VALUE(bIsAdvancedDisplay, bInIsAdvancedDisplay))
+		if (const UMetaSoundEditorViewModel* MetaSoundEditorViewModel = Cast<UMetaSoundEditorViewModel>(GetOuter()))
 		{
-			DocumentBuilder.SetGraphOutputAdvancedDisplay(OutputName, bInIsAdvancedDisplay);
+			MetaSoundEditorViewModel->SetOutputIsAdvancedDisplay(OutputName, bInIsAdvancedDisplay);
 		}
 	}
 }
diff --git a/Engine/Plugins/Experimental/TechAudioTools/Source/TechAudioToolsMetaSoundEditor/Public/ViewModels/MetaSoundEditorViewModel.h b/Engine/Plugins/Experimental/TechAudioTools/Source/TechAudioToolsMetaSoundEditor/Public/ViewModels/MetaSoundEditorViewModel.h
--- a/Engine/Plugins/Experimental/TechAudioTools/Source/TechAudioToolsMetaSoundEditor/Public/ViewModels/MetaSoundEditorViewModel.h
+++ b/Engine/Plugins/Experimental/TechAudioTools/Source/TechAudioToolsMetaSoundEditor/Public/ViewModels/MetaSoundEditorViewModel.h
@@ -6,9 +6,9 @@
 
 #include "MetaSoundEditorViewModel.generated.h"
 
-class UMetaSoundEditorBuilderListener;
+#define UE_API TECHAUDIOTOOLSMETASOUNDEDITOR_API
 
-TECHAUDIOTOOLSMETASOUNDEDITOR_API DECLARE_LOG_CATEGORY_EXTERN(LogTechAudioToolsMetaSoundEditor, Log, All);
+class UMetaSoundEditorBuilderListener;
 
 /**
  * Editor viewmodel for MetaSounds. Creates MetaSoundEditorBuilderListener bindings upon initialization, allowing changes made to assets in the
@@ -21,89 +21,103 @@ class UMetaSoundEditorViewModel : public UMetaSoundViewModel
 
 protected:
 	// Sets the display name of the initialized MetaSound.
-	UPROPERTY(BlueprintReadWrite, FieldNotify, Setter = "SetMetaSoundDisplayName", Category = "MetaSound Viewmodel|Metadata", meta = (AllowPrivateAccess))
+	UPROPERTY(BlueprintReadWrite, FieldNotify, Getter, Setter = "SetMetaSoundDisplayName", Category = "MetaSound Viewmodel|Metadata", meta = (AllowPrivateAccess))
 	FText DisplayName;
 
 	// Sets the description of the initialized MetaSound.
-	UPROPERTY(BlueprintReadWrite, FieldNotify, Setter = "SetMetaSoundDescription", Category = "MetaSound Viewmodel|Metadata", meta = (AllowPrivateAccess))
+	UPROPERTY(BlueprintReadWrite, FieldNotify, Getter, Setter = "SetMetaSoundDescription", Category = "MetaSound Viewmodel|Metadata", meta = (AllowPrivateAccess))
 	FText Description;
 
 	// Sets the author of the initialized MetaSound.
-	UPROPERTY(BlueprintReadWrite, FieldNotify, Setter, Category = "MetaSound Viewmodel|Metadata", meta = (AllowPrivateAccess))
+	UPROPERTY(BlueprintReadWrite, FieldNotify, Getter, Setter, Category = "MetaSound Viewmodel|Metadata", meta = (AllowPrivateAccess))
 	FString Author;
 
 	// Sets the keywords of the initialized MetaSound.
-	UPROPERTY(BlueprintReadWrite, FieldNotify, Setter, Category = "MetaSound Viewmodel|Metadata", meta = (AllowPrivateAccess))
+	UPROPERTY(BlueprintReadWrite, FieldNotify, Getter, Setter, Category = "MetaSound Viewmodel|Metadata", meta = (AllowPrivateAccess))
 	TArray<FText> Keywords;
 
 	// Sets the category hierarchy of the initialized MetaSound.
-	UPROPERTY(BlueprintReadWrite, FieldNotify, Setter, Category = "MetaSound Viewmodel|Metadata", meta = (AllowPrivateAccess))
+	UPROPERTY(BlueprintReadWrite, FieldNotify, Getter, Setter, Category = "MetaSound Viewmodel|Metadata", meta = (AllowPrivateAccess))
 	TArray<FText> CategoryHierarchy;
 
 	// Sets the initialized MetaSound asset as deprecated.
-	UPROPERTY(BlueprintReadWrite, FieldNotify, Setter = "SetIsDeprecated", Category = "MetaSound Viewmodel|Metadata", meta = (AllowPrivateAccess))
+	UPROPERTY(BlueprintReadWrite, FieldNotify, Getter = "IsDeprecated", Setter = "SetIsDeprecated", Category = "MetaSound Viewmodel|Metadata", meta = (AllowPrivateAccess))
 	bool bIsDeprecated = false;
 
 public:
-	virtual bool IsEditorOnly() const override;
-	virtual UWorld* GetWorld() const override;
+	virtual bool IsEditorOnly() const override { return true; }
 
 	virtual void InitializeMetaSound(const TScriptInterface<IMetaSoundDocumentInterface> InMetaSound) override;
 	virtual void Initialize(UMetaSoundBuilderBase* InBuilder) override;
 	virtual void Reset() override;
 
-	virtual void InitializeProperties(const FMetasoundFrontendDocument& FrontendDocument) override;
-	virtual void ResetProperties() override;
-
-	void SetMetaSoundDisplayName(const FText& InDisplayName);
-	void SetMetaSoundDescription(const FText& InDescription);
-	void SetAuthor(const FString& InAuthor);
-	void SetKeywords(const TArray<FText>& InKeywords);
-	void SetCategoryHierarchy(const TArray<FText>& InCategoryHierarchy);
-	void SetIsDeprecated(const bool bInIsDeprecated);
-
-protected:
-	UPROPERTY(Transient)
-	TObjectPtr<UMetaSoundEditorBuilderListener> BuilderListener;
-
-	virtual void CreateInputViewModel(const FMetasoundFrontendClassInput& InInput) override;
-    virtual UMetaSoundInputViewModel* CreateInputViewModelInstance() override;
-    virtual UMetaSoundOutputViewModel* CreateOutputViewModelInstance() override;
-
-	void BindDelegates();
-	void UnbindDelegates();
+	FText GetDisplayName() const { return DisplayName; }
+	FText GetDescription() const { return Description; };
+	FString GetAuthor() const { return Author; };
+	TArray<FText> GetKeywords() const { return Keywords; };
+	TArray<FText> GetCategoryHierarchy() const { return CategoryHierarchy; };
+	bool IsDeprecated() const { return bIsDeprecated; };
+
+	UE_API void SetMetaSoundDisplayName(const FText& InDisplayName);
+	UE_API void SetMetaSoundDescription(const FText& InDescription);
+	UE_API void SetAuthor(const FString& InAuthor);
+	UE_API void SetKeywords(const TArray<FText>& InKeywords);
+	UE_API void SetCategoryHierarchy(const TArray<FText>& InCategoryHierarchy);
+	UE_API void SetIsDeprecated(bool bInIsDeprecated);
+
+	UE_API void SetInputDisplayName(const FName& InputName, const FText& InDisplayName) const;
+	UE_API void SetInputDescription(const FName& InputName, const FText& InDescription) const;
+	UE_API void SetInputSortOrderIndex(const FName& InputName, int32 InSortOrderIndex) const;
+	UE_API void SetInputIsAdvancedDisplay(const FName& InputName, bool bInIsAdvancedDisplay) const;
+
+	UE_API void SetOutputDisplayName(const FName& OutputName, const FText& InDisplayName) const;
+	UE_API void SetOutputDescription(const FName& OutputName, const FText& InDescription) const;
+	UE_API void SetOutputSortOrderIndex(const FName& OutputName, int32 InSortOrderIndex) const;
+	UE_API void SetOutputIsAdvancedDisplay(const FName& OutputName, bool bInIsAdvancedDisplay) const;
 
-	// Called when a new MetaSound input has been added to the initialized MetaSound.
 	UFUNCTION()
-	void OnInputAdded(FName VertexName, FName DataType);
-
-	// Called when a MetaSound input has been removed from the initialized MetaSound.
+	void OnDisplayNameChanged(const FText NewDisplayName) { UE_MVVM_SET_PROPERTY_VALUE(DisplayName, NewDisplayName); }
 	UFUNCTION()
-	void OnInputRemoved(FName VertexName, FName DataType);
-
-	// Called when the name of an input on the initialized MetaSound has changed.
+	void OnDescriptionChanged(const FText NewDescription) { UE_MVVM_SET_PROPERTY_VALUE(Description, NewDescription); }
 	UFUNCTION()
-	void OnInputNameChanged(FName OldName, FName NewName);
-
-	// Called when an input's data type has changed.
+	void OnAuthorChanged(const FString NewAuthor) { UE_MVVM_SET_PROPERTY_VALUE(Author, NewAuthor); }
 	UFUNCTION()
-	void OnInputDataTypeChanged(FName VertexName, FName DataType);
-
-	// Called when a new MetaSound output has been added to the initialized MetaSound.
+	void OnKeywordsChanged(const TArray<FText> NewKeywords) { UE_MVVM_SET_PROPERTY_VALUE(Keywords, NewKeywords); }
 	UFUNCTION()
-	void OnOutputAdded(FName VertexName, FName DataType);
-
-	// Called when a MetaSound output has been removed from the initialized MetaSound.
+	void OnCategoryHierarchyChanged(const TArray<FText> NewCategoryHierarchy) { UE_MVVM_SET_PROPERTY_VALUE(CategoryHierarchy, NewCategoryHierarchy); }
 	UFUNCTION()
-	void OnOutputRemoved(FName VertexName, FName DataType);
+	void OnIsDeprecatedChanged(const bool bNewIsDeprecated) { UE_MVVM_SET_PROPERTY_VALUE(bIsDeprecated, bNewIsDeprecated); }
 
-	// Called when the name of an output on the initialized MetaSound has changed.
 	UFUNCTION()
-	void OnOutputNameChanged(FName OldName, FName NewName);
+	void OnInputDisplayNameChanged(FName InputName, FText InDisplayName);
+	UFUNCTION()
+	void OnInputDescriptionChanged(FName InputName, FText InDescription);
+	UFUNCTION()
+	void OnInputSortOrderIndexChanged(FName InputName, int32 InSortOrderIndex);
+	UFUNCTION()
+	void OnInputIsAdvancedDisplayChanged(FName InputName, bool bInIsAdvancedDisplay);
 
-	// Called when an output's data type has changed.
 	UFUNCTION()
-	void OnOutputDataTypeChanged(FName VertexName, FName DataType);
+	void OnOutputDisplayNameChanged(FName OutputName, const FText InDisplayName);
+	UFUNCTION()
+	void OnOutputDescriptionChanged(FName OutputName, const FText InDescription);
+	UFUNCTION()
+	void OnOutputSortOrderIndexChanged(FName OutputName, int32 InSortOrderIndex);
+	UFUNCTION()
+	void OnOutputIsAdvancedDisplayChanged(FName OutputName, bool bInIsAdvancedDisplay);
+
+protected:
+	UPROPERTY(Transient)
+	TObjectPtr<UMetaSoundEditorBuilderListener> BuilderListener;
+
+	virtual UMetaSoundInputViewModel* CreateInputViewModel(const FMetasoundFrontendClassInput& InInput) override;
+	virtual UMetaSoundOutputViewModel* CreateOutputViewModel(const FMetasoundFrontendClassOutput& InOutput) override;
+
+	virtual TSubclassOf<UMetaSoundInputViewModel> GetInputViewModelClass() const override;
+	virtual TSubclassOf<UMetaSoundOutputViewModel> GetOutputViewModelClass() const override;
+
+	void BindDelegates();
+	void UnbindDelegates();
 };
 
 /**
@@ -132,25 +146,22 @@ protected:
 	bool bIsAdvancedDisplay;
 
 public:
-	virtual bool IsEditorOnly() const override;
-	virtual UWorld* GetWorld() const override;
-
-	virtual void InitializeInput(UMetaSoundBuilderBase* InBuilder, const FMetasoundFrontendClassInput& InInput) override;
-	virtual void ResetInput() override;
+	virtual bool IsEditorOnly() const override { return true; }
 
 	FText GetInputDisplayName() const { return InputDisplayName; }
 	FText GetInputDescription() const { return InputDescription; }
 	int32 GetSortOrderIndex() const { return SortOrderIndex; }
 	bool IsAdvancedDisplay() const { return bIsAdvancedDisplay; }
 
-	void SetInputDisplayName(const FText& InDisplayName);
-	void SetInputDescription(const FText& InDescription);
-	void SetSortOrderIndex(const int32 InSortOrderIndex);
-	void SetIsAdvancedDisplay(const bool bInIsAdvancedDisplay);
+	UE_API void SetInputDisplayName(const FText& InDisplayName);
+	UE_API void SetInputDescription(const FText& InDescription);
+	UE_API void SetSortOrderIndex(int32 InSortOrderIndex);
+	UE_API void SetIsAdvancedDisplay(bool bInIsAdvancedDisplay);
 
-	// Called when the default value of an input has been changed on the initialized MetaSound.
-	UFUNCTION()
-	void OnInputDefaultChanged(FName VertexName, FMetasoundFrontendLiteral LiteralValue, FName PageName);
+	void OnInputDisplayNameChanged(const FText& NewDisplayName) { UE_MVVM_SET_PROPERTY_VALUE(InputDisplayName, NewDisplayName); }
+	void OnInputDescriptionChanged(const FText& NewDescription) { UE_MVVM_SET_PROPERTY_VALUE(InputDescription, NewDescription); }
+	void OnInputSortOrderIndexChanged(const int32 NewSortOrderIndex) { UE_MVVM_SET_PROPERTY_VALUE(SortOrderIndex, NewSortOrderIndex); }
+	void OnInputIsAdvancedDisplayChanged(const bool bNewIsAdvancedDisplay) { UE_MVVM_SET_PROPERTY_VALUE(bIsAdvancedDisplay, bNewIsAdvancedDisplay); }
 };
 
 /**
@@ -179,19 +190,22 @@ protected:
 	bool bIsAdvancedDisplay;
 
 public:
-	virtual bool IsEditorOnly() const override;
-	virtual UWorld* GetWorld() const override;
-
-	virtual void InitializeOutput(UMetaSoundBuilderBase* InBuilder, const FMetasoundFrontendClassOutput& InOutput) override;
-	virtual void ResetOutput() override;
+	virtual bool IsEditorOnly() const override { return true; }
 
 	FText GetOutputDisplayName() const { return OutputDisplayName; }
 	FText GetOutputDescription() const { return OutputDescription; }
 	int32 GetSortOrderIndex() const { return SortOrderIndex; }
 	bool IsAdvancedDisplay() const { return bIsAdvancedDisplay; }
 
-	void SetOutputDisplayName(const FText& InDisplayName);
-	void SetOutputDescription(const FText& InDescription);
-	void SetSortOrderIndex(const int32 InSortOrderIndex);
-	void SetIsAdvancedDisplay(const bool bInIsAdvancedDisplay);
+	UE_API void SetOutputDisplayName(const FText& InDisplayName);
+	UE_API void SetOutputDescription(const FText& InDescription);
+	UE_API void SetSortOrderIndex(int32 InSortOrderIndex);
+	UE_API void SetIsAdvancedDisplay(bool bInIsAdvancedDisplay);
+
+	void OnOutputDisplayNameChanged(const FText& NewDisplayName) { UE_MVVM_SET_PROPERTY_VALUE(OutputDisplayName, NewDisplayName); }
+	void OnOutputDescriptionChanged(const FText& NewDescription) { UE_MVVM_SET_PROPERTY_VALUE(OutputDescription, NewDescription); }
+	void OnOutputSortOrderIndexChanged(const int32 NewSortOrderIndex) { UE_MVVM_SET_PROPERTY_VALUE(SortOrderIndex, NewSortOrderIndex); }
+	void OnOutputIsAdvancedDisplayChanged(const bool NewIsAdvancedDisplay) { UE_MVVM_SET_PROPERTY_VALUE(bIsAdvancedDisplay, NewIsAdvancedDisplay); }
 };
+
+#undef UE_API
diff --git a/Engine/Plugins/Experimental/TechAudioTools/Source/TechAudioToolsMetaSoundEditor/TechAudioToolsMetaSoundEditor.Build.cs b/Engine/Plugins/Experimental/TechAudioTools/Source/TechAudioToolsMetaSoundEditor/TechAudioToolsMetaSoundEditor.Build.cs
--- a/Engine/Plugins/Experimental/TechAudioTools/Source/TechAudioToolsMetaSoundEditor/TechAudioToolsMetaSoundEditor.Build.cs
+++ b/Engine/Plugins/Experimental/TechAudioTools/Source/TechAudioToolsMetaSoundEditor/TechAudioToolsMetaSoundEditor.Build.cs
@@ -18,6 +18,7 @@ public class TechAudioToolsMetaSoundEditor : ModuleRules
 				"MetasoundFrontend",
 				"ModelViewViewModel",
 				"ModelViewViewModelEditor",
+				"TechAudioTools",
 				"TechAudioToolsMetaSound",
 				"UnrealEd",
             }
