TechAudioTools fix for SourceRange not being converted from DisplayRange for Volume and Pitch range remapping

#rnx

diff --git a/Engine/Plugins/Experimental/TechAudioTools/Source/TechAudioTools/Private/TechAudioToolsFloatMapping.cpp b/Engine/Plugins/Experimental/TechAudioTools/Source/TechAudioTools/Private/TechAudioToolsFloatMapping.cpp
--- a/Engine/Plugins/Experimental/TechAudioTools/Source/TechAudioTools/Private/TechAudioToolsFloatMapping.cpp
+++ b/Engine/Plugins/Experimental/TechAudioTools/Source/TechAudioTools/Private/TechAudioToolsFloatMapping.cpp
@@ -16,14 +16,24 @@ FFloatInterval UTechAudioToolsFloatMapping::GetRange(const ETechAudioToolsMappin
 
 	case ETechAudioToolsFloatMappingType::Volume:
 		{
-			const bool bShowDecibels = Endpoint == ETechAudioToolsMappingEndpoint::Display ? DisplayVolumeUnits == ETechAudioToolsVolumeUnit::Decibels : SourceVolumeUnits == ETechAudioToolsVolumeUnit::Decibels;
-			return bShowDecibels ? DisplayRange_Decibels : DisplayRange_LinearGain;
+			const FFloatInterval VolumeDisplayRange = DisplayVolumeUnits == ETechAudioToolsVolumeUnit::Decibels ? DisplayRange_Decibels : DisplayRange_LinearGain;
+			if (Endpoint == ETechAudioToolsMappingEndpoint::Display || SourceVolumeUnits == DisplayVolumeUnits)
+			{
+				return VolumeDisplayRange;
+			}
+
+			return TechAudioTools::ConvertRange<ETechAudioToolsVolumeUnit>(DisplayVolumeUnits, SourceVolumeUnits, VolumeDisplayRange);
 		}
 
 	case ETechAudioToolsFloatMappingType::Pitch:
 		{
-			const bool bShowSemitones = Endpoint == ETechAudioToolsMappingEndpoint::Display ? DisplayPitchUnits == ETechAudioToolsPitchUnit::Semitones : SourcePitchUnits == ETechAudioToolsPitchUnit::Semitones;
-			return bShowSemitones ? DisplayRange_Semitones : DisplayRange_FrequencyMultiplier;
+			const FFloatInterval PitchDisplayRange = DisplayPitchUnits == ETechAudioToolsPitchUnit::Semitones ? DisplayRange_Semitones : DisplayRange_FrequencyMultiplier;
+			if (Endpoint == ETechAudioToolsMappingEndpoint::Display || SourcePitchUnits == DisplayPitchUnits)
+			{
+				return PitchDisplayRange;
+			}
+
+			return TechAudioTools::ConvertRange<ETechAudioToolsPitchUnit>(DisplayPitchUnits, SourcePitchUnits, PitchDisplayRange);
 		}
 	}
 
diff --git a/Engine/Plugins/Experimental/TechAudioTools/Source/TechAudioTools/Public/TechAudioToolsFloatMapping.h b/Engine/Plugins/Experimental/TechAudioTools/Source/TechAudioTools/Public/TechAudioToolsFloatMapping.h
--- a/Engine/Plugins/Experimental/TechAudioTools/Source/TechAudioTools/Public/TechAudioToolsFloatMapping.h
+++ b/Engine/Plugins/Experimental/TechAudioTools/Source/TechAudioTools/Public/TechAudioToolsFloatMapping.h
@@ -10,25 +10,44 @@
 namespace TechAudioTools
 {
 	// Converts volume values between linear gain and decibels
-	FORCEINLINE float ConvertUnit(const ETechAudioToolsVolumeUnit From, const ETechAudioToolsVolumeUnit To, const float Value)
+	FORCEINLINE float ConvertUnit(const ETechAudioToolsVolumeUnit FromUnits, const ETechAudioToolsVolumeUnit ToUnits, const float Value)
 	{
-		if (From == To)
+		if (FromUnits == ToUnits)
 		{
 			return Value;
 		}
 
-		return To == ETechAudioToolsVolumeUnit::Decibels ? Audio::ConvertToDecibels(Value) : Audio::ConvertToLinear(Value);
+		return ToUnits == ETechAudioToolsVolumeUnit::Decibels ? Audio::ConvertToDecibels(Value) : Audio::ConvertToLinear(Value);
 	}
 
 	// Converts pitch values between semitones and frequency multiplier
-	FORCEINLINE float ConvertUnit(const ETechAudioToolsPitchUnit From, const ETechAudioToolsPitchUnit To, const float Value)
+	FORCEINLINE float ConvertUnit(const ETechAudioToolsPitchUnit FromUnits, const ETechAudioToolsPitchUnit ToUnits, const float Value)
 	{
-		if (From == To)
+		if (FromUnits == ToUnits)
 		{
 			return Value;
 		}
 
-		return To == ETechAudioToolsPitchUnit::Semitones ? Audio::GetSemitones(Value) : Audio::GetFrequencyMultiplier(Value);
+		return ToUnits == ETechAudioToolsPitchUnit::Semitones ? Audio::GetSemitones(Value) : Audio::GetFrequencyMultiplier(Value);
+	}
+
+	template <typename UnitT>
+	FORCEINLINE FFloatInterval ConvertRange(UnitT FromUnits, UnitT ToUnits, FFloatInterval InputRange)
+	{
+		if (FromUnits == ToUnits)
+		{
+			return InputRange;
+		}
+
+		float RangeMin = ConvertUnit(FromUnits, ToUnits, InputRange.Min);
+		float RangeMax = ConvertUnit(FromUnits, ToUnits, InputRange.Max);
+
+		if (RangeMin > RangeMax)
+		{
+			Swap(RangeMin, RangeMax);
+		}
+
+		return FFloatInterval(RangeMin, RangeMax);
 	}
 } // namespace TechAudioTools
 
@@ -133,5 +152,5 @@ public:
 	TECHAUDIOTOOLS_API float DisplayToNormalized(float InDisplay) const;
 
 private:
-	FFloatInterval GetRange(const ETechAudioToolsMappingEndpoint Endpoint) const;
+	TECHAUDIOTOOLS_API FFloatInterval GetRange(const ETechAudioToolsMappingEndpoint Endpoint) const;
 };
