[1/3] TechAudioTools - added CommonUI dependency and exposed number formatting

Removed TechAudioToolsStyle since number formattings are now exposed instead of hardcoded

#rnx

diff --git a/Engine/Plugins/Experimental/TechAudioTools/Source/TechAudioTools/Private/Viewmodels/TechAudioToolsConversionFunctions.cpp b/Engine/Plugins/Experimental/TechAudioTools/Source/TechAudioTools/Private/Viewmodels/TechAudioToolsConversionFunctions.cpp
--- a/Engine/Plugins/Experimental/TechAudioTools/Source/TechAudioTools/Private/Viewmodels/TechAudioToolsConversionFunctions.cpp
+++ b/Engine/Plugins/Experimental/TechAudioTools/Source/TechAudioTools/Private/Viewmodels/TechAudioToolsConversionFunctions.cpp
@@ -2,47 +2,53 @@
 
 #include "Viewmodels/TechAudioToolsConversionFunctions.h"
 
-#include "TechAudioToolsStyle.h"
+#include "CommonNumericTextBlock.h"
 
 #include UE_INLINE_GENERATED_CPP_BY_NAME(TechAudioToolsConversionFunctions)
 
-FText UTechAudioToolsConversionFunctions::GetFloatValueFormatted(const float InValue, const ETechAudioToolsFloatUnit DisplayUnit)
+FText UTechAudioToolsConversionFunctions::GetFloatValueFormatted(const float InValue, const ETechAudioToolsFloatUnit DisplayUnit, const FCommonNumberFormattingOptions InFormat)
 {
-	using namespace TechAudioTools;
-
 	FText FormattedFloat;
 	FText FloatUnitLabel;
 
+	FNumberFormattingOptions Format;
+	Format.RoundingMode = InFormat.RoundingMode;
+	Format.AlwaysSign = InFormat.AlwaysSign;
+	Format.UseGrouping = InFormat.UseGrouping;
+	Format.MinimumIntegralDigits = InFormat.MinimumIntegralDigits;
+	Format.MinimumFractionalDigits = InFormat.MinimumFractionalDigits;
+	Format.MaximumFractionalDigits = InFormat.MaximumFractionalDigits;
+
 	switch (DisplayUnit)
 	{
 	case ETechAudioToolsFloatUnit::None:
-		FormattedFloat = FText::AsNumber(InValue, GetDefaultFloatFormat());
+		FormattedFloat = FText::AsNumber(InValue, &Format);
 		break;
 	case ETechAudioToolsFloatUnit::Decibels:
-		FormattedFloat = FText::AsNumber(InValue, GetDecibelsFloatFormat());
+		FormattedFloat = FText::AsNumber(InValue, &Format);
 		FloatUnitLabel = FText::FromString(TEXT(" dB"));
 		break;
 	case ETechAudioToolsFloatUnit::Hertz:
-		FormattedFloat = FText::AsNumber(InValue, GetFrequencyFloatFormat());
+		FormattedFloat = FText::AsNumber(InValue, &Format);
 		FloatUnitLabel = FText::FromString(TEXT(" Hz"));
 		break;
 	case ETechAudioToolsFloatUnit::Semitones:
-		FormattedFloat = FText::AsNumber(InValue, GetSemitonesFloatFormat());
+		FormattedFloat = FText::AsNumber(InValue, &Format);
 		break;
 	case ETechAudioToolsFloatUnit::TimeSeconds:
-		FormattedFloat = FText::AsNumber(InValue, GetSecondsFloatFormat());
+		FormattedFloat = FText::AsNumber(InValue, &Format);
 		FloatUnitLabel = FText::FromString(TEXT("s"));
 		break;
 	case ETechAudioToolsFloatUnit::TimeMilliseconds:
-		FormattedFloat = FText::AsNumber(InValue, GetMillisecondsFloatFormat());
+		FormattedFloat = FText::AsNumber(InValue, &Format);
 		FloatUnitLabel = FText::FromString(TEXT("ms"));
 		break;
 	case ETechAudioToolsFloatUnit::Degrees:
-		FormattedFloat = FText::AsNumber(InValue, GetDegreesFloatFormat());
+		FormattedFloat = FText::AsNumber(InValue, &Format);
 		FloatUnitLabel = FText::FromString(TEXT("°"));
 		break;
 	case ETechAudioToolsFloatUnit::Multiplier:
-		FormattedFloat = FText::AsNumber(InValue, GetMultiplierFloatFormat());
+		FormattedFloat = FText::AsNumber(InValue, &Format);
 		FloatUnitLabel = FText::FromString(TEXT("x"));
 		break;
 	case ETechAudioToolsFloatUnit::Pan:
@@ -54,12 +60,12 @@ FText UTechAudioToolsConversionFunctions::GetFloatValueFormatted(const float InV
 		}
 		else if (InValue < 0.f)
 		{
-			FormattedFloat  = FText::AsNumber(FMath::Abs(InValue), GetPanFloatFormat());
+			FormattedFloat  = FText::AsNumber(FMath::Abs(InValue), &Format);
 			FloatUnitLabel  = FText::FromString(TEXT(" L"));
 		}
 		else
 		{
-			FormattedFloat  = FText::AsNumber(FMath::Abs(InValue), GetPanFloatFormat());
+			FormattedFloat  = FText::AsNumber(FMath::Abs(InValue), &Format);
 			FloatUnitLabel  = FText::FromString(TEXT(" R"));
 		}
 		break;
diff --git a/Engine/Plugins/Experimental/TechAudioTools/Source/TechAudioTools/Public/TechAudioToolsStyle.h b/Engine/Plugins/Experimental/TechAudioTools/Source/TechAudioTools/Public/TechAudioToolsStyle.h
deleted file mode 100644
diff --git a/Engine/Plugins/Experimental/TechAudioTools/Source/TechAudioTools/Public/Viewmodels/TechAudioToolsConversionFunctions.h b/Engine/Plugins/Experimental/TechAudioTools/Source/TechAudioTools/Public/Viewmodels/TechAudioToolsConversionFunctions.h
--- a/Engine/Plugins/Experimental/TechAudioTools/Source/TechAudioTools/Public/Viewmodels/TechAudioToolsConversionFunctions.h
+++ b/Engine/Plugins/Experimental/TechAudioTools/Source/TechAudioTools/Public/Viewmodels/TechAudioToolsConversionFunctions.h
@@ -7,13 +7,15 @@
 
 #include "TechAudioToolsConversionFunctions.generated.h"
 
+struct FCommonNumberFormattingOptions;
+
 UCLASS(MinimalAPI)
 class UTechAudioToolsConversionFunctions : public UBlueprintFunctionLibrary
 {
 	GENERATED_BODY()
 
 public:
-	// Provides a float value as formatted text using the given display unit.
+	// Provides a float value as formatted text using the given display unit and format.
 	UFUNCTION(BlueprintCallable, BlueprintPure, Category = "TechAudioTools", DisplayName = "Get Float Value Formatted (TechAudioTools)")
-	static TECHAUDIOTOOLS_API FText GetFloatValueFormatted(float InValue, ETechAudioToolsFloatUnit DisplayUnit);
+	static TECHAUDIOTOOLS_API FText GetFloatValueFormatted(float InValue, ETechAudioToolsFloatUnit DisplayUnit, FCommonNumberFormattingOptions InFormat);
 };
diff --git a/Engine/Plugins/Experimental/TechAudioTools/Source/TechAudioTools/TechAudioTools.Build.cs b/Engine/Plugins/Experimental/TechAudioTools/Source/TechAudioTools/TechAudioTools.Build.cs
--- a/Engine/Plugins/Experimental/TechAudioTools/Source/TechAudioTools/TechAudioTools.Build.cs
+++ b/Engine/Plugins/Experimental/TechAudioTools/Source/TechAudioTools/TechAudioTools.Build.cs
@@ -18,6 +18,7 @@ public class TechAudioTools : ModuleRules
         PrivateDependencyModuleNames.AddRange(
             new string[]
             {
+	            "CommonUI",
                 "CoreUObject",
                 "Engine",
             }
diff --git a/Engine/Plugins/Experimental/TechAudioTools/Source/TechAudioToolsMetaSound/TechAudioToolsMetaSound.Build.cs b/Engine/Plugins/Experimental/TechAudioTools/Source/TechAudioToolsMetaSound/TechAudioToolsMetaSound.Build.cs
--- a/Engine/Plugins/Experimental/TechAudioTools/Source/TechAudioToolsMetaSound/TechAudioToolsMetaSound.Build.cs
+++ b/Engine/Plugins/Experimental/TechAudioTools/Source/TechAudioToolsMetaSound/TechAudioToolsMetaSound.Build.cs
@@ -11,6 +11,7 @@ public class TechAudioToolsMetaSound : ModuleRules
         PublicDependencyModuleNames.AddRange(
             new string[]
             {
+	            "CommonUI",
                 "Core",
 				"MetasoundEngine",
 				"MetasoundFrontend",
diff --git a/Engine/Plugins/Experimental/TechAudioTools/TechAudioTools.uplugin b/Engine/Plugins/Experimental/TechAudioTools/TechAudioTools.uplugin
--- a/Engine/Plugins/Experimental/TechAudioTools/TechAudioTools.uplugin
+++ b/Engine/Plugins/Experimental/TechAudioTools/TechAudioTools.uplugin
@@ -16,22 +16,26 @@
 	"Installed": false,
 	"Modules": [
 		{
-			"Name": "TechAudioToolsMetaSound",
+			"Name": "TechAudioTools",
 			"Type": "Runtime",
 			"LoadingPhase": "Default"
 		},
 		{
-			"Name": "TechAudioToolsMetaSoundEditor",
-			"Type": "Editor",
+			"Name": "TechAudioToolsMetaSound",
+			"Type": "Runtime",
 			"LoadingPhase": "Default"
 		},
 		{
-			"Name": "TechAudioTools",
-			"Type": "Runtime",
+			"Name": "TechAudioToolsMetaSoundEditor",
+			"Type": "Editor",
 			"LoadingPhase": "Default"
 		}
 	],
 	"Plugins": [
+		{
+			"Name": "CommonUI",
+			"Enabled": true
+		},
 		{
 			"Name": "Metasound",
 			"Enabled": true
