[2/3] TechAudioTools - Added MetaSound Literal Widget Interface and float viewmodel

#redirectsaresafe

diff --git a/Engine/Plugins/Experimental/TechAudioTools/Source/TechAudioToolsMetaSound/Private/ViewModels/MetaSoundLiteralViewModel.cpp b/Engine/Plugins/Experimental/TechAudioTools/Source/TechAudioToolsMetaSound/Private/ViewModels/MetaSoundLiteralViewModel.cpp
new file mode 100644
--- /dev/null
+++ b/Engine/Plugins/Experimental/TechAudioTools/Source/TechAudioToolsMetaSound/Private/ViewModels/MetaSoundLiteralViewModel.cpp
@@ -0,0 +1,145 @@
+﻿// Copyright Epic Games, Inc. All Rights Reserved.
+
+#include "ViewModels/MetaSoundLiteralViewModel.h"
+
+#include "Viewmodels/TechAudioToolsConversionFunctions.h"
+#include "TechAudioToolsFloatMapping.h"
+
+#include UE_INLINE_GENERATED_CPP_BY_NAME(MetaSoundLiteralViewModel)
+
+float UMetaSoundLiteralViewModel_Float::GetSourceRangeMin() const
+{
+	return RangeValues ? RangeValues->GetSourceMin() : 0.f;
+}
+
+float UMetaSoundLiteralViewModel_Float::GetSourceRangeMax() const
+{
+	return RangeValues ? RangeValues->GetSourceMax() : 0.f;
+}
+
+ETechAudioToolsFloatUnit UMetaSoundLiteralViewModel_Float::GetSourceUnits() const
+{
+	return RangeValues ? RangeValues->GetUnits(ETechAudioToolsMappingEndpoint::Source) : ETechAudioToolsFloatUnit::None;
+}
+
+float UMetaSoundLiteralViewModel_Float::GetDisplayRangeMin() const
+{
+	return RangeValues ? RangeValues->GetDisplayMin() : 0.f;
+}
+
+float UMetaSoundLiteralViewModel_Float::GetDisplayRangeMax() const
+{
+	return RangeValues ? RangeValues->GetDisplayMax() : 0.f;
+}
+
+ETechAudioToolsFloatUnit UMetaSoundLiteralViewModel_Float::GetDisplayUnits() const
+{
+	return RangeValues ? RangeValues->GetUnits(ETechAudioToolsMappingEndpoint::Display) : ETechAudioToolsFloatUnit::None;
+}
+
+UMetaSoundLiteralViewModel_Float::UMetaSoundLiteralViewModel_Float(const FObjectInitializer& ObjectInitializer)
+{
+	RangeValues = ObjectInitializer.CreateDefaultSubobject<UTechAudioToolsFloatMapping>(this, TEXT("RangeValues"));
+}
+
+void UMetaSoundLiteralViewModel_Float::SetShowDisplayValues(const bool bInIsShowingDisplayValues)
+{
+	if (!bIsUpdating && RangeValues)
+	{
+		bIsUpdating = true;
+		bShowDisplayValues = bInIsShowingDisplayValues;
+		NormalizedValue = bShowDisplayValues ? RangeValues->DisplayToNormalized(DisplayValue) : RangeValues->SourceToNormalized(SourceValue);
+		UE_MVVM_BROADCAST_FIELD_VALUE_CHANGED(NormalizedValue);
+		UE_MVVM_BROADCAST_FIELD_VALUE_CHANGED(bShowDisplayValues);
+		bIsUpdating = false;
+	}
+}
+
+void UMetaSoundLiteralViewModel_Float::SetLiteral(const FMetasoundFrontendLiteral& InLiteral)
+{
+	const bool bInitializing = Literal.GetType() == EMetasoundFrontendLiteralType::None;
+
+	if (!bIsUpdating && RangeValues)
+	{
+		bIsUpdating = true;
+		Literal = InLiteral;
+		InLiteral.TryGet(SourceValue);
+		DisplayValue = RangeValues->SourceToDisplay(SourceValue);
+		NormalizedValue = bShowDisplayValues ? RangeValues->DisplayToNormalized(DisplayValue) : RangeValues->SourceToNormalized(SourceValue);
+
+		// Don't broadcast Literal changed when setting the literal for the first time to avoid overriding inherited value
+		if (!bInitializing)
+		{
+			UE_MVVM_BROADCAST_FIELD_VALUE_CHANGED(Literal);
+		}
+		UE_MVVM_BROADCAST_FIELD_VALUE_CHANGED(SourceValue);
+		UE_MVVM_BROADCAST_FIELD_VALUE_CHANGED(NormalizedValue);
+		UE_MVVM_BROADCAST_FIELD_VALUE_CHANGED(DisplayValue);
+		bIsUpdating = false;
+	}
+}
+
+void UMetaSoundLiteralViewModel_Float::SetSourceValue(const float InSourceValue)
+{
+	if (!bIsUpdating && RangeValues)
+	{
+		bIsUpdating = true;
+		SourceValue = InSourceValue;
+		DisplayValue = RangeValues->SourceToDisplay(SourceValue);
+		NormalizedValue = bShowDisplayValues ? RangeValues->DisplayToNormalized(DisplayValue) : RangeValues->SourceToNormalized(SourceValue);
+		Literal.Set(SourceValue);
+
+		UE_MVVM_BROADCAST_FIELD_VALUE_CHANGED(Literal);
+		UE_MVVM_BROADCAST_FIELD_VALUE_CHANGED(SourceValue);
+		UE_MVVM_BROADCAST_FIELD_VALUE_CHANGED(NormalizedValue);
+		UE_MVVM_BROADCAST_FIELD_VALUE_CHANGED(DisplayValue);
+		bIsUpdating = false;
+	}
+}
+
+void UMetaSoundLiteralViewModel_Float::SetNormalizedValue(const float InNormalizedValue)
+{
+	if (!bIsUpdating && RangeValues)
+	{
+		bIsUpdating = true;
+		NormalizedValue = InNormalizedValue;
+
+		// Properly ensures the normalized range is mapped correctly
+		if (bShowDisplayValues)
+		{
+			DisplayValue = RangeValues->NormalizedToDisplay(InNormalizedValue);
+			SourceValue = RangeValues->DisplayToSource(DisplayValue);
+		}
+		else
+		{
+			SourceValue = RangeValues->NormalizedToSource(InNormalizedValue);
+			DisplayValue = RangeValues->SourceToDisplay(SourceValue);
+		}
+
+		Literal.Set(SourceValue);
+
+		UE_MVVM_BROADCAST_FIELD_VALUE_CHANGED(Literal);
+		UE_MVVM_BROADCAST_FIELD_VALUE_CHANGED(SourceValue);
+		UE_MVVM_BROADCAST_FIELD_VALUE_CHANGED(NormalizedValue);
+		UE_MVVM_BROADCAST_FIELD_VALUE_CHANGED(DisplayValue);
+		bIsUpdating = false;
+	}
+}
+
+void UMetaSoundLiteralViewModel_Float::SetDisplayValue(const float InDisplayValue)
+{
+	if (!bIsUpdating && RangeValues)
+	{
+		bIsUpdating = true;
+		DisplayValue = InDisplayValue;
+		SourceValue = RangeValues->DisplayToSource(InDisplayValue);
+		NormalizedValue = bShowDisplayValues ? RangeValues->DisplayToNormalized(DisplayValue) : RangeValues->SourceToNormalized(SourceValue);
+		Literal.Set(SourceValue);
+
+		UE_MVVM_BROADCAST_FIELD_VALUE_CHANGED(Literal);
+		UE_MVVM_BROADCAST_FIELD_VALUE_CHANGED(SourceValue);
+		UE_MVVM_BROADCAST_FIELD_VALUE_CHANGED(NormalizedValue);
+		UE_MVVM_BROADCAST_FIELD_VALUE_CHANGED(DisplayValue);
+		bIsUpdating = false;
+	}
+}
diff --git a/Engine/Plugins/Experimental/TechAudioTools/Source/TechAudioToolsMetaSound/Public/Interfaces/MetaSoundLiteralInterface.h b/Engine/Plugins/Experimental/TechAudioTools/Source/TechAudioToolsMetaSound/Public/Interfaces/MetaSoundLiteralInterface.h
new file mode 100644
--- /dev/null
+++ b/Engine/Plugins/Experimental/TechAudioTools/Source/TechAudioToolsMetaSound/Public/Interfaces/MetaSoundLiteralInterface.h
@@ -0,0 +1,47 @@
+﻿// Copyright Epic Games, Inc. All Rights Reserved.
+
+#pragma once
+
+#include "UObject/Interface.h"
+
+#include "MetaSoundLiteralInterface.generated.h"
+
+class UMetaSoundInputViewModel;
+class UMetaSoundViewModel;
+
+/**
+ * Used for retrieving specific viewmodels from a MetaSoundViewModel, such as inputs or outputs. Simplifies the process of initializing MetaSound
+ * Literal Widgets by allowing the owning widget to call functions on entire containers of widgets that implement this interface, rather than
+ * painstakingly initializing each widget one-by-one.
+ *
+ * Widgets implementing this interface are expected to provide an array of input names in order to request the required InputViewModels from the owning
+ * widget's MetaSoundViewModel. Then, the widget can implement SetInputViewModels to correctly assign each InputViewModel received.
+ *
+ * To assign viewmodels, it is recommended to use either UMVVMView::SetViewModelByClass (if there aren't multiple viewmodels of the same class) or
+ * UMVVMView::SetViewModel (if there are multiple viewmodels of the same class). If using UMVVMView::SetViewModel, you must ensure the name of the
+ * viewmodel passed into the function matches the name of the corresponding viewmodel instance *exactly* or else the assignment will fail. It is
+ * recommended to expose a name variable to designers rather than hardcode a name value.
+ *
+ * @see UMetaSoundLiteralWidget_Float for an example implementation.
+ */
+UINTERFACE(MinimalAPI, Blueprintable, DisplayName = "MetaSound Literal Widget Interface")
+class UMetaSoundLiteralWidgetInterface : public UInterface
+{
+	GENERATED_BODY()
+};
+
+class IMetaSoundLiteralWidgetInterface
+{
+	GENERATED_BODY()
+
+public:
+	// Returns an array containing the names of the MetaSound Input Viewmodels required for this widget.
+	UFUNCTION(BlueprintCallable, BlueprintNativeEvent, DisplayName = "Get Input Viewmodel Names", Category = "MetaSound Literal")
+	TArray<FName> GetInputViewModelNames() const;
+	virtual TArray<FName> GetInputViewModelNames_Implementation() const { return TArray<FName>(); }
+
+	// Returns the InputViewModels required for this widget.
+	UFUNCTION(BlueprintCallable, BlueprintNativeEvent, DisplayName = "Set Input Viewmodels", Category = "MetaSound Literal")
+	void SetInputViewModels(UPARAM(DisplayName = "Input Viewmodels") const TMap<FName, UMetaSoundInputViewModel*>& InputViewModels);
+	virtual void SetInputViewModels_Implementation(const TMap<FName, UMetaSoundInputViewModel*>& InputViewModels) {}
+};
diff --git a/Engine/Plugins/Experimental/TechAudioTools/Source/TechAudioToolsMetaSound/Public/ViewModels/MetaSoundLiteralViewModel.h b/Engine/Plugins/Experimental/TechAudioTools/Source/TechAudioToolsMetaSound/Public/ViewModels/MetaSoundLiteralViewModel.h
new file mode 100644
--- /dev/null
+++ b/Engine/Plugins/Experimental/TechAudioTools/Source/TechAudioToolsMetaSound/Public/ViewModels/MetaSoundLiteralViewModel.h
@@ -0,0 +1,101 @@
+// Copyright Epic Games, Inc. All Rights Reserved.
+
+#pragma once
+
+#include "MetasoundFrontendLiteral.h"
+#include "MVVMViewModelBase.h"
+#include "TechAudioToolsTypes.h"
+
+#include "MetaSoundLiteralViewModel.generated.h"
+
+class UMetaSoundBuilderBase;
+class UTechAudioToolsFloatMapping;
+
+/**
+ * Viewmodel for converting a float-type MetaSound Literal between Source, Normalized, and Display values.
+ *
+ * Source values refer to the units or range expected by the MetaSound Literal. Display values refer to a configurable conversion of the Source value.
+ * For instance, a "Volume" input might expect a linear gain value in the MetaSound, but users may wish to display the value as decibels.
+ *
+ * The Normalized value is to be used for UI elements such as sliders and knobs.
+ *
+ * The Literal variable belonging to this viewmodel is expected to be bound to the Literal variable of a MetaSoundInputViewModel using a two-way
+ * binding for bidirectional updates. Doing so will ensure all updates to the input in the MetaSound Editor are reflected in your widget and vice versa.
+ *
+ * This viewmodel should most likely use the Create Instance setting, along with the Expose Instance in Editor option. This allows designers to easily
+ * configure units and range settings in the Details panel of instances of the owning widget.
+ *
+ * @see UTechAudioToolsFloatMapping for unit conversion and range mapping logic.
+ */
+UCLASS(BlueprintType, EditInlineNew, DefaultToInstanced, CollapseCategories, DisplayName = "MetaSound Literal Float Viewmodel")
+class UMetaSoundLiteralViewModel_Float : public UMVVMViewModelBase
+{
+	GENERATED_BODY()
+
+protected:
+	// If true, the NormalizedValue will be mapped to the DisplayRange, else mapped to the SourceRange.
+	UPROPERTY(BlueprintReadWrite, EditAnywhere, FieldNotify, Getter = "IsShowingDisplayValues", Setter = "SetShowDisplayValues", Category = "Float Literal")
+	bool bShowDisplayValues = true;
+
+	// The bound MetaSound Literal value. Should be bound using a two-way binding with the Literal variable of a MetaSoundInputViewModel.
+	UPROPERTY(BlueprintReadWrite, FieldNotify, Setter, Category = "Float Literal", meta = (AllowPrivateAccess))
+	FMetasoundFrontendLiteral Literal;
+
+	// The value of the bound Literal as a float.
+	UPROPERTY(BlueprintReadWrite, FieldNotify, Getter, Setter, Category = "Float Literal", meta = (AllowPrivateAccess))
+	float SourceValue = 0.f;
+
+	// The value of the DisplayValue in normalized (0-1) space. Used for knob/slider visual appearance.
+	UPROPERTY(BlueprintReadWrite, FieldNotify, Getter, Setter, Category = "Float Literal", meta = (AllowPrivateAccess))
+	float NormalizedValue = 0.f;
+
+	// The user-facing display value (e.g. dB, Hz, semitones). Can be modified and is automatically converted to the expected SourceRange.
+	UPROPERTY(BlueprintReadWrite, FieldNotify, Getter, Setter, Category = "Float Literal", meta = (AllowPrivateAccess))
+	float DisplayValue = 0.f;
+
+	// Mapping pair used to convert between float literal, normalized, and display values.
+	UPROPERTY(EditAnywhere, Instanced, Category = "Float Range Settings", meta = (AllowPrivateAccess))
+	TObjectPtr<UTechAudioToolsFloatMapping> RangeValues;
+
+public:
+	// Returns the range minimum of the SourceRange.
+	UFUNCTION(BlueprintCallable, BlueprintPure, FieldNotify, Category = "TechAudioTools")
+	TECHAUDIOTOOLSMETASOUND_API float GetSourceRangeMin() const;
+
+	// Returns the range maximum of the SourceRange.
+	UFUNCTION(BlueprintCallable, BlueprintPure, FieldNotify, Category = "TechAudioTools")
+	TECHAUDIOTOOLSMETASOUND_API float GetSourceRangeMax() const;
+
+	// Returns the units associated with the SourceValue.
+	UFUNCTION(BlueprintCallable, BlueprintPure, FieldNotify, Category = "TechAudioTools")
+	TECHAUDIOTOOLSMETASOUND_API ETechAudioToolsFloatUnit GetSourceUnits() const;
+
+	// Returns the range minimum of the DisplayRange.
+	UFUNCTION(BlueprintCallable, BlueprintPure, FieldNotify, Category = "TechAudioTools")
+	TECHAUDIOTOOLSMETASOUND_API float GetDisplayRangeMin() const;
+
+	// Returns the range maximum of the DisplayRange.
+	UFUNCTION(BlueprintCallable, BlueprintPure, FieldNotify, Category = "TechAudioTools")
+	TECHAUDIOTOOLSMETASOUND_API float GetDisplayRangeMax() const;
+
+	// Returns the units associated with the DisplayValue.
+	UFUNCTION(BlueprintCallable, BlueprintPure, FieldNotify, Category = "TechAudioTools")
+	TECHAUDIOTOOLSMETASOUND_API ETechAudioToolsFloatUnit GetDisplayUnits() const;
+
+	explicit UMetaSoundLiteralViewModel_Float(const FObjectInitializer& ObjectInitializer);
+
+	bool IsShowingDisplayValues() const { return bShowDisplayValues; }
+	FMetasoundFrontendLiteral GetLiteral() const { return Literal; }
+	float GetSourceValue() const { return SourceValue; }
+	float GetNormalizedValue() const { return NormalizedValue; }
+	float GetDisplayValue() const { return DisplayValue; }
+
+	TECHAUDIOTOOLSMETASOUND_API void SetShowDisplayValues(bool bInIsShowingDisplayValues);
+	TECHAUDIOTOOLSMETASOUND_API void SetLiteral(const FMetasoundFrontendLiteral& InLiteral);
+	TECHAUDIOTOOLSMETASOUND_API void SetSourceValue(float InSourceValue);
+	TECHAUDIOTOOLSMETASOUND_API void SetNormalizedValue(float InNormalizedValue);
+	TECHAUDIOTOOLSMETASOUND_API void SetDisplayValue(float InDisplayValue);
+
+private:
+	bool bIsUpdating = false;
+};
