TechAudioTools - documentation tweaks and fix for potential null dereference

#rnx

diff --git a/Engine/Plugins/Experimental/TechAudioTools/Source/TechAudioToolsMetaSound/Public/Interfaces/MetaSoundLiteralInterface.h b/Engine/Plugins/Experimental/TechAudioTools/Source/TechAudioToolsMetaSound/Public/Interfaces/MetaSoundLiteralInterface.h
--- a/Engine/Plugins/Experimental/TechAudioTools/Source/TechAudioToolsMetaSound/Public/Interfaces/MetaSoundLiteralInterface.h
+++ b/Engine/Plugins/Experimental/TechAudioTools/Source/TechAudioToolsMetaSound/Public/Interfaces/MetaSoundLiteralInterface.h
@@ -40,7 +40,7 @@ public:
 	TArray<FName> GetInputViewModelNames() const;
 	virtual TArray<FName> GetInputViewModelNames_Implementation() const { return TArray<FName>(); }
 
-	// Returns the InputViewModels required for this widget.
+	// Sets the MetaSound input viewmodels this widget binds to.
 	UFUNCTION(BlueprintCallable, BlueprintNativeEvent, DisplayName = "Set Input Viewmodels", Category = "MetaSound Literal")
 	void SetInputViewModels(UPARAM(DisplayName = "Input Viewmodels") const TMap<FName, UMetaSoundInputViewModel*>& InputViewModels);
 	virtual void SetInputViewModels_Implementation(const TMap<FName, UMetaSoundInputViewModel*>& InputViewModels) {}
diff --git a/Engine/Plugins/Experimental/TechAudioTools/Source/TechAudioToolsMetaSound/Public/ViewModels/MetaSoundLiteralViewModel.h b/Engine/Plugins/Experimental/TechAudioTools/Source/TechAudioToolsMetaSound/Public/ViewModels/MetaSoundLiteralViewModel.h
--- a/Engine/Plugins/Experimental/TechAudioTools/Source/TechAudioToolsMetaSound/Public/ViewModels/MetaSoundLiteralViewModel.h
+++ b/Engine/Plugins/Experimental/TechAudioTools/Source/TechAudioToolsMetaSound/Public/ViewModels/MetaSoundLiteralViewModel.h
@@ -45,7 +45,10 @@ protected:
 	UPROPERTY(BlueprintReadWrite, FieldNotify, Getter, Setter, Category = "Float Literal", meta = (AllowPrivateAccess))
 	float SourceValue = 0.f;
 
-	// The value of the DisplayValue in normalized (0-1) space. Used for knob/slider visual appearance.
+	/**
+	 * The normalized (0-1) representation of the value currently surfaced to the UI. If bShowDisplayValues is true, this is DisplayValue mapped into
+	 * 0-1 space. Otherwise, it's SourceValue mapped into 0-1 space. Used by knobs/sliders purely for visual feedback.
+	 */
 	UPROPERTY(BlueprintReadWrite, FieldNotify, Getter, Setter, Category = "Float Literal", meta = (AllowPrivateAccess))
 	float NormalizedValue = 0.f;
 
