TechAudioTools MetaSoundViewModel - prevent member viewmodels from being re-instanced

#rnx

diff --git a/Engine/Plugins/Experimental/TechAudioTools/Source/TechAudioToolsMetaSound/Private/ViewModels/MetaSoundViewModel.cpp b/Engine/Plugins/Experimental/TechAudioTools/Source/TechAudioToolsMetaSound/Private/ViewModels/MetaSoundViewModel.cpp
--- a/Engine/Plugins/Experimental/TechAudioTools/Source/TechAudioToolsMetaSound/Private/ViewModels/MetaSoundViewModel.cpp
+++ b/Engine/Plugins/Experimental/TechAudioTools/Source/TechAudioToolsMetaSound/Private/ViewModels/MetaSoundViewModel.cpp
@@ -18,6 +18,11 @@ TArray<UMetaSoundInputViewModel*> UMetaSoundViewModel::GetInputViewModels() cons
 	return InputViewModelsArray;
 }
 
+UMetaSoundInputViewModel* UMetaSoundViewModel::FindInputViewModel(const FName InputViewModelName) const
+{
+	return InputViewModels.FindRef(InputViewModelName);
+}
+
 TArray<UMetaSoundOutputViewModel*> UMetaSoundViewModel::GetOutputViewModels() const
 {
 	TArray<TObjectPtr<UMetaSoundOutputViewModel>> OutputViewModelsArray;
@@ -25,6 +30,11 @@ TArray<UMetaSoundOutputViewModel*> UMetaSoundViewModel::GetOutputViewModels() co
 	return OutputViewModelsArray;
 }
 
+UMetaSoundOutputViewModel* UMetaSoundViewModel::FindOutputViewModel(const FName OutputViewModelName) const
+{
+	return OutputViewModels.FindRef(OutputViewModelName);
+}
+
 void UMetaSoundViewModel::InitializeMetaSound(const TScriptInterface<IMetaSoundDocumentInterface> InMetaSound)
 {
 	if (!InMetaSound)
@@ -219,7 +229,7 @@ void UMetaSoundViewModel::SetOutputIsConstructorPin(const FName& OutputName, con
 
 void UMetaSoundViewModel::OnInputAdded(const FName VertexName, FName DataType)
 {
-	if (Builder)
+	if (Builder && !InputViewModels.Contains(VertexName))
 	{
 		const FMetaSoundFrontendDocumentBuilder& DocumentBuilder = Builder->GetConstBuilder();
 		const FMetasoundFrontendDocument& FrontendDocument = DocumentBuilder.GetConstDocumentChecked();
@@ -289,7 +299,7 @@ void UMetaSoundViewModel::OnInputIsConstructorPinChanged(const FName VertexName,
 
 void UMetaSoundViewModel::OnOutputAdded(const FName VertexName, FName DataType)
 {
-	if (Builder)
+	if (Builder && !OutputViewModels.Contains(VertexName))
 	{
 		const FMetaSoundFrontendDocumentBuilder& DocumentBuilder = Builder->GetConstBuilder();
 		const FMetasoundFrontendDocument& FrontendDocument = DocumentBuilder.GetConstDocumentChecked();
diff --git a/Engine/Plugins/Experimental/TechAudioTools/Source/TechAudioToolsMetaSound/Public/ViewModels/MetaSoundViewModel.h b/Engine/Plugins/Experimental/TechAudioTools/Source/TechAudioToolsMetaSound/Public/ViewModels/MetaSoundViewModel.h
--- a/Engine/Plugins/Experimental/TechAudioTools/Source/TechAudioToolsMetaSound/Public/ViewModels/MetaSoundViewModel.h
+++ b/Engine/Plugins/Experimental/TechAudioTools/Source/TechAudioToolsMetaSound/Public/ViewModels/MetaSoundViewModel.h
@@ -41,10 +41,18 @@ public:
 	UFUNCTION(BlueprintCallable, FieldNotify, DisplayName = "Get Input Viewmodels", Category = "MetaSound Viewmodel")
 	UE_API virtual TArray<UMetaSoundInputViewModel*> GetInputViewModels() const;
 
-	// Contains MetaSound Output ViewModels for each output of the initialized MetaSound.
+	// Returns the input viewmodel with the given name if it exists, else returns nullptr.
+	UFUNCTION(BlueprintCallable, DisplayName = "Find Input Viewmodel", Category = "MetaSound Viewmodel")
+	UE_API UMetaSoundInputViewModel* FindInputViewModel(const FName InputViewModelName) const;
+
+	// Contains MetaSound Output Viewmodels for each output of the initialized MetaSound.
 	UFUNCTION(BlueprintCallable, FieldNotify, DisplayName = "Get Output Viewmodels", Category = "MetaSound Viewmodel")
 	UE_API virtual TArray<UMetaSoundOutputViewModel*> GetOutputViewModels() const;
 
+	// Returns the output viewmodel with the given name if it exists, else returns nullptr.
+	UFUNCTION(BlueprintCallable, DisplayName = "Find Output Viewmodel", Category = "MetaSound Viewmodel")
+	UE_API UMetaSoundOutputViewModel* FindOutputViewModel(const FName OutputViewModelName) const;
+
 	// Initializes the viewmodel using the given MetaSound asset.
 	UFUNCTION(BlueprintCallable, DisplayName = "Initialize MetaSound", Category = "Audio|MetaSound Viewmodel")
 	UE_API virtual void InitializeMetaSound(const TScriptInterface<IMetaSoundDocumentInterface> InMetaSound);
