TechAudioTools - added check in MetaSoundEditorViewModel to prevent possible null dereference

#rnx

diff --git a/Engine/Plugins/Experimental/TechAudioTools/Source/TechAudioToolsMetaSoundEditor/Private/ViewModels/MetaSoundEditorViewModel.cpp b/Engine/Plugins/Experimental/TechAudioTools/Source/TechAudioToolsMetaSoundEditor/Private/ViewModels/MetaSoundEditorViewModel.cpp
--- a/Engine/Plugins/Experimental/TechAudioTools/Source/TechAudioToolsMetaSoundEditor/Private/ViewModels/MetaSoundEditorViewModel.cpp
+++ b/Engine/Plugins/Experimental/TechAudioTools/Source/TechAudioToolsMetaSoundEditor/Private/ViewModels/MetaSoundEditorViewModel.cpp
@@ -29,8 +29,15 @@ void UMetaSoundEditorViewModel::InitializeMetaSound(const TScriptInterface<IMeta
 
 void UMetaSoundEditorViewModel::Initialize(UMetaSoundBuilderBase* InBuilder)
 {
+	Reset();
+
 	Super::Initialize(InBuilder);
 
+	if (!InBuilder)
+	{
+		return;
+	}
+
 	const FMetaSoundFrontendDocumentBuilder& DocumentBuilder = InBuilder->GetConstBuilder();
 	const FMetasoundFrontendDocument& FrontendDocument = DocumentBuilder.GetConstDocumentChecked();
 
