TechAudioTools - added viewmodel classes for each MetaSound literal type

#rnx

diff --git a/Engine/Plugins/Experimental/TechAudioTools/Source/TechAudioToolsMetaSound/Private/ViewModels/MetaSoundLiteralViewModel.cpp b/Engine/Plugins/Experimental/TechAudioTools/Source/TechAudioToolsMetaSound/Private/ViewModels/MetaSoundLiteralViewModel.cpp
--- a/Engine/Plugins/Experimental/TechAudioTools/Source/TechAudioToolsMetaSound/Private/ViewModels/MetaSoundLiteralViewModel.cpp
+++ b/Engine/Plugins/Experimental/TechAudioTools/Source/TechAudioToolsMetaSound/Private/ViewModels/MetaSoundLiteralViewModel.cpp
@@ -2,11 +2,196 @@
 
 #include "ViewModels/MetaSoundLiteralViewModel.h"
 
-#include "Viewmodels/TechAudioToolsConversionFunctions.h"
 #include "TechAudioToolsFloatMapping.h"
 
 #include UE_INLINE_GENERATED_CPP_BY_NAME(MetaSoundLiteralViewModel)
 
+void UMetaSoundLiteralViewModel_Boolean::SetLiteral(const FMetasoundFrontendLiteral& InLiteral)
+{
+	if (UE_MVVM_SET_PROPERTY_VALUE(Literal, InLiteral))
+	{
+		Literal = InLiteral;
+		InLiteral.TryGet(SourceValue);
+		UE_MVVM_BROADCAST_FIELD_VALUE_CHANGED(SourceValue);
+	}
+}
+
+void UMetaSoundLiteralViewModel_Boolean::SetSourceValue(const bool InSourceValue)
+{
+	if (UE_MVVM_SET_PROPERTY_VALUE(SourceValue, InSourceValue))
+	{
+		Literal.Set(InSourceValue);
+		UE_MVVM_BROADCAST_FIELD_VALUE_CHANGED(Literal);
+	}
+}
+
+void UMetaSoundLiteralViewModel_BooleanArray::SetLiteral(const FMetasoundFrontendLiteral& InLiteral)
+{
+	if (UE_MVVM_SET_PROPERTY_VALUE(Literal, InLiteral))
+	{
+		Literal = InLiteral;
+		InLiteral.TryGet(SourceValue);
+		UE_MVVM_BROADCAST_FIELD_VALUE_CHANGED(SourceValue);
+	}
+}
+
+void UMetaSoundLiteralViewModel_BooleanArray::SetSourceValue(const TArray<bool>& InSourceValue)
+{
+	if (UE_MVVM_SET_PROPERTY_VALUE(SourceValue, InSourceValue))
+	{
+		Literal.Set(InSourceValue);
+		UE_MVVM_BROADCAST_FIELD_VALUE_CHANGED(Literal);
+	}
+}
+
+float UMetaSoundLiteralViewModel_Integer::GetStepSize() const
+{
+	if (!Range.HasLowerBound() || !Range.HasUpperBound())
+	{
+		return 1.f;
+	}
+
+	if (!Range.GetLowerBound().IsClosed() || !Range.GetUpperBound().IsClosed())
+	{
+		return 1.f;
+	}
+
+	int32 RangeMin = Range.GetLowerBoundValue();
+	int32 RangeMax = Range.GetUpperBoundValue();
+
+	if (Range.GetLowerBound().IsExclusive())
+	{
+		RangeMin += 1;
+	}
+	if (Range.GetUpperBound().IsExclusive())
+	{
+		RangeMax -= 1;
+	}
+
+	const int32 Span = RangeMax - RangeMin;
+	if (Span <= 0)
+	{
+		return 1.f;
+	}
+
+	const float Step = 1.f / Span;
+	return FMath::Clamp(Step, 0.f, 1.f);
+}
+
+void UMetaSoundLiteralViewModel_Integer::SetLiteral(const FMetasoundFrontendLiteral& InLiteral)
+{
+	const bool bInitializing = Literal.GetType() == EMetasoundFrontendLiteralType::None;
+
+	if (!bIsUpdating)
+	{
+		bIsUpdating = true;
+		Literal = InLiteral;
+		InLiteral.TryGet(SourceValue);
+
+		if (!Range.HasLowerBound() || !Range.HasUpperBound())
+		{
+			NormalizedValue = SourceValue;
+		}
+		else
+		{
+			const float RangeMin = Range.GetLowerBoundValue();
+			const float RangeMax = Range.GetUpperBoundValue();
+			NormalizedValue = FMath::GetRangePct(RangeMin, RangeMax, SourceValue);
+		}
+
+		// Don't broadcast Literal changed when setting the literal for the first time to avoid overriding inherited value
+		if (!bInitializing)
+		{
+			UE_MVVM_BROADCAST_FIELD_VALUE_CHANGED(Literal);
+		}
+		UE_MVVM_BROADCAST_FIELD_VALUE_CHANGED(SourceValue);
+		UE_MVVM_BROADCAST_FIELD_VALUE_CHANGED(NormalizedValue);
+		bIsUpdating = false;
+	}
+}
+
+void UMetaSoundLiteralViewModel_Integer::SetSourceValue(const int32 InSourceValue)
+{
+	if (!bIsUpdating)
+	{
+		bIsUpdating = true;
+		SourceValue = InSourceValue;
+		Literal.Set(SourceValue);
+
+		if (!Range.HasLowerBound() || !Range.HasUpperBound())
+		{
+			NormalizedValue = SourceValue;
+		}
+		else
+		{
+			const float RangeMin = Range.GetLowerBoundValue();
+			const float RangeMax = Range.GetUpperBoundValue();
+			NormalizedValue = FMath::GetRangePct(RangeMin, RangeMax, SourceValue);
+		}
+
+		UE_MVVM_BROADCAST_FIELD_VALUE_CHANGED(Literal);
+		UE_MVVM_BROADCAST_FIELD_VALUE_CHANGED(SourceValue);
+		UE_MVVM_BROADCAST_FIELD_VALUE_CHANGED(NormalizedValue);
+		bIsUpdating = false;
+	}
+}
+
+void UMetaSoundLiteralViewModel_Integer::SetNormalizedValue(const float InNormalizedValue)
+{
+	if (!bIsUpdating)
+	{
+		bIsUpdating = true;
+		NormalizedValue = InNormalizedValue;
+
+		int32 RangeMin = 0;
+		int32 RangeMax = 1;
+		if (Range.HasLowerBound() && Range.HasUpperBound())
+		{
+			RangeMin = Range.GetLowerBoundValue();
+			RangeMax = Range.GetUpperBoundValue();
+
+			if (Range.GetLowerBound().IsExclusive())
+			{
+				RangeMin += 1;
+			}
+			if (Range.GetUpperBound().IsExclusive())
+			{
+				RangeMax -= 1;
+			}
+		}
+
+		const FVector2D NormalizedRange(0, 1);
+		const float UnmappedFloat = FMath::GetMappedRangeValueUnclamped(NormalizedRange, FVector2D(RangeMin, RangeMax), NormalizedValue);
+		SourceValue = FMath::RoundToInt(UnmappedFloat);
+
+		Literal.Set(SourceValue);
+
+		UE_MVVM_BROADCAST_FIELD_VALUE_CHANGED(Literal);
+		UE_MVVM_BROADCAST_FIELD_VALUE_CHANGED(SourceValue);
+		UE_MVVM_BROADCAST_FIELD_VALUE_CHANGED(NormalizedValue);
+		bIsUpdating = false;
+	}
+}
+
+void UMetaSoundLiteralViewModel_IntegerArray::SetLiteral(const FMetasoundFrontendLiteral& InLiteral)
+{
+	if (UE_MVVM_SET_PROPERTY_VALUE(Literal, InLiteral))
+	{
+		Literal = InLiteral;
+		InLiteral.TryGet(SourceValue);
+		UE_MVVM_BROADCAST_FIELD_VALUE_CHANGED(SourceValue);
+	}
+}
+
+void UMetaSoundLiteralViewModel_IntegerArray::SetSourceValue(const TArray<int32>& InSourceValue)
+{
+	if (UE_MVVM_SET_PROPERTY_VALUE(SourceValue, InSourceValue))
+	{
+		Literal.Set(InSourceValue);
+		UE_MVVM_BROADCAST_FIELD_VALUE_CHANGED(Literal);
+	}
+}
+
 float UMetaSoundLiteralViewModel_Float::GetSourceRangeMin() const
 {
 	return RangeValues ? RangeValues->GetSourceMin() : 0.f;
@@ -143,3 +328,102 @@ void UMetaSoundLiteralViewModel_Float::SetDisplayValue(const float InDisplayValu
 		bIsUpdating = false;
 	}
 }
+
+void UMetaSoundLiteralViewModel_FloatArray::SetLiteral(const FMetasoundFrontendLiteral& InLiteral)
+{
+	if (UE_MVVM_SET_PROPERTY_VALUE(Literal, InLiteral))
+	{
+		Literal = InLiteral;
+		InLiteral.TryGet(SourceValue);
+		UE_MVVM_BROADCAST_FIELD_VALUE_CHANGED(SourceValue);
+	}
+}
+
+void UMetaSoundLiteralViewModel_FloatArray::SetSourceValue(const TArray<float>& InSourceValue)
+{
+	if (UE_MVVM_SET_PROPERTY_VALUE(SourceValue, InSourceValue))
+	{
+		Literal.Set(InSourceValue);
+		UE_MVVM_BROADCAST_FIELD_VALUE_CHANGED(Literal);
+	}
+}
+
+void UMetaSoundLiteralViewModel_String::SetLiteral(const FMetasoundFrontendLiteral& InLiteral)
+{
+	if (UE_MVVM_SET_PROPERTY_VALUE(Literal, InLiteral))
+	{
+		Literal = InLiteral;
+		InLiteral.TryGet(SourceValue);
+		UE_MVVM_BROADCAST_FIELD_VALUE_CHANGED(SourceValue);
+	}
+}
+
+void UMetaSoundLiteralViewModel_String::SetSourceValue(const FString& InSourceValue)
+{
+	if (UE_MVVM_SET_PROPERTY_VALUE(SourceValue, InSourceValue))
+	{
+		Literal.Set(InSourceValue);
+		UE_MVVM_BROADCAST_FIELD_VALUE_CHANGED(Literal);
+	}
+}
+
+void UMetaSoundLiteralViewModel_StringArray::SetLiteral(const FMetasoundFrontendLiteral& InLiteral)
+{
+	if (UE_MVVM_SET_PROPERTY_VALUE(Literal, InLiteral))
+	{
+		Literal = InLiteral;
+		InLiteral.TryGet(SourceValue);
+		UE_MVVM_BROADCAST_FIELD_VALUE_CHANGED(SourceValue);
+	}
+}
+
+void UMetaSoundLiteralViewModel_StringArray::SetSourceValue(const TArray<FString>& InSourceValue)
+{
+	if (UE_MVVM_SET_PROPERTY_VALUE(SourceValue, InSourceValue))
+	{
+		Literal.Set(InSourceValue);
+		UE_MVVM_BROADCAST_FIELD_VALUE_CHANGED(Literal);
+	}
+}
+
+void UMetaSoundLiteralViewModel_Object::SetLiteral(const FMetasoundFrontendLiteral& InLiteral)
+{
+	if (UE_MVVM_SET_PROPERTY_VALUE(Literal, InLiteral))
+	{
+		Literal = InLiteral;
+		UObject* SourceValuePtr = nullptr;
+		InLiteral.TryGet(SourceValuePtr);
+		SourceValue = SourceValuePtr;
+		UE_MVVM_BROADCAST_FIELD_VALUE_CHANGED(SourceValue);
+	}
+}
+
+void UMetaSoundLiteralViewModel_Object::SetSourceValue(UObject* InSourceValue)
+{
+	if (UE_MVVM_SET_PROPERTY_VALUE(SourceValue, InSourceValue))
+	{
+		Literal.Set(InSourceValue);
+		UE_MVVM_BROADCAST_FIELD_VALUE_CHANGED(Literal);
+	}
+}
+
+void UMetaSoundLiteralViewModel_ObjectArray::SetLiteral(const FMetasoundFrontendLiteral& InLiteral)
+{
+	if (UE_MVVM_SET_PROPERTY_VALUE(Literal, InLiteral))
+	{
+		Literal = InLiteral;
+		TArray<UObject*> SourceValuePtr;
+		InLiteral.TryGet(SourceValuePtr);
+		SourceValue = SourceValuePtr;
+		UE_MVVM_BROADCAST_FIELD_VALUE_CHANGED(SourceValue);
+	}
+}
+
+void UMetaSoundLiteralViewModel_ObjectArray::SetSourceValue(const TArray<TObjectPtr<UObject>>& InSourceValue)
+{
+	if (UE_MVVM_SET_PROPERTY_VALUE(SourceValue, InSourceValue))
+	{
+		Literal.Set(InSourceValue);
+		UE_MVVM_BROADCAST_FIELD_VALUE_CHANGED(Literal);
+	}
+}
diff --git a/Engine/Plugins/Experimental/TechAudioTools/Source/TechAudioToolsMetaSound/Public/Interfaces/MetaSoundLiteralInterface.h b/Engine/Plugins/Experimental/TechAudioTools/Source/TechAudioToolsMetaSound/Public/Interfaces/MetaSoundLiteralInterface.h
--- a/Engine/Plugins/Experimental/TechAudioTools/Source/TechAudioToolsMetaSound/Public/Interfaces/MetaSoundLiteralInterface.h
+++ b/Engine/Plugins/Experimental/TechAudioTools/Source/TechAudioToolsMetaSound/Public/Interfaces/MetaSoundLiteralInterface.h
@@ -9,6 +9,8 @@
 class UMetaSoundInputViewModel;
 class UMetaSoundViewModel;
 
+#define UE_API TECHAUDIOTOOLSMETASOUND_API
+
 /**
  * Used for retrieving specific viewmodels from a MetaSoundViewModel, such as inputs or outputs. Simplifies the process of initializing MetaSound
  * Literal Widgets by allowing the owning widget to call functions on entire containers of widgets that implement this interface, rather than
@@ -37,11 +39,13 @@ class IMetaSoundLiteralWidgetInterface
 public:
 	// Returns an array containing the names of the MetaSound Input Viewmodels required for this widget.
 	UFUNCTION(BlueprintCallable, BlueprintNativeEvent, DisplayName = "Get Input Viewmodel Names", Category = "MetaSound Literal")
-	TArray<FName> GetInputViewModelNames() const;
+	UE_API TArray<FName> GetInputViewModelNames() const;
 	virtual TArray<FName> GetInputViewModelNames_Implementation() const { return TArray<FName>(); }
 
 	// Sets the MetaSound input viewmodels this widget binds to.
 	UFUNCTION(BlueprintCallable, BlueprintNativeEvent, DisplayName = "Set Input Viewmodels", Category = "MetaSound Literal")
-	void SetInputViewModels(UPARAM(DisplayName = "Input Viewmodels") const TMap<FName, UMetaSoundInputViewModel*>& InputViewModels);
+	UE_API void SetInputViewModels(UPARAM(DisplayName = "Input Viewmodels") const TMap<FName, UMetaSoundInputViewModel*>& InputViewModels);
 	virtual void SetInputViewModels_Implementation(const TMap<FName, UMetaSoundInputViewModel*>& InputViewModels) {}
 };
+
+#undef UE_API
diff --git a/Engine/Plugins/Experimental/TechAudioTools/Source/TechAudioToolsMetaSound/Public/ViewModels/MetaSoundLiteralViewModel.h b/Engine/Plugins/Experimental/TechAudioTools/Source/TechAudioToolsMetaSound/Public/ViewModels/MetaSoundLiteralViewModel.h
--- a/Engine/Plugins/Experimental/TechAudioTools/Source/TechAudioToolsMetaSound/Public/ViewModels/MetaSoundLiteralViewModel.h
+++ b/Engine/Plugins/Experimental/TechAudioTools/Source/TechAudioToolsMetaSound/Public/ViewModels/MetaSoundLiteralViewModel.h
@@ -11,6 +11,124 @@
 class UMetaSoundBuilderBase;
 class UTechAudioToolsFloatMapping;
 
+#define UE_API TECHAUDIOTOOLSMETASOUND_API
+
+/**
+ * Viewmodel for binding a boolean-type MetaSound Literal with a widget.
+ */
+UCLASS(BlueprintType, EditInlineNew, DefaultToInstanced, CollapseCategories, DisplayName = "MetaSound Literal Boolean Viewmodel")
+class UMetaSoundLiteralViewModel_Boolean : public UMVVMViewModelBase
+{
+	GENERATED_BODY()
+
+protected:
+	// The bound MetaSound Literal value. Should be bound using a two-way binding with the Literal variable of a MetaSoundInputViewModel.
+	UPROPERTY(BlueprintReadWrite, Transient, FieldNotify, Setter, Category = "Boolean Literal", meta = (AllowPrivateAccess))
+	FMetasoundFrontendLiteral Literal;
+
+	// The value of the bound Literal as a bool.
+	UPROPERTY(BlueprintReadWrite, Transient, FieldNotify, Getter, Setter, Category = "Boolean Literal", meta = (AllowPrivateAccess))
+	bool SourceValue = false;
+
+public:
+	const FMetasoundFrontendLiteral& GetLiteral() const { return Literal; }
+	bool GetSourceValue() const { return SourceValue; }
+
+	UE_API void SetLiteral(const FMetasoundFrontendLiteral& InLiteral);
+	UE_API void SetSourceValue(bool InSourceValue);
+};
+
+/**
+ * Viewmodel for binding a boolean array-type MetaSound Literal with a widget.
+ */
+UCLASS(BlueprintType, EditInlineNew, DefaultToInstanced, CollapseCategories, DisplayName = "MetaSound Literal Boolean Array Viewmodel")
+class UMetaSoundLiteralViewModel_BooleanArray : public UMVVMViewModelBase
+{
+	GENERATED_BODY()
+
+protected:
+	// The bound MetaSound Literal value. Should be bound using a two-way binding with the Literal variable of a MetaSoundInputViewModel.
+	UPROPERTY(BlueprintReadWrite, Transient, FieldNotify, Setter, Category = "Boolean Array Literal", meta = (AllowPrivateAccess))
+	FMetasoundFrontendLiteral Literal;
+
+	// The value of the bound Literal as a boolean array.
+	UPROPERTY(BlueprintReadWrite, Transient, FieldNotify, Getter, Setter, Category = "Boolean Array Literal", meta = (AllowPrivateAccess))
+	TArray<bool> SourceValue;
+
+public:
+	const FMetasoundFrontendLiteral& GetLiteral() const { return Literal; }
+	TArray<bool> GetSourceValue() const { return SourceValue; }
+
+	UE_API void SetLiteral(const FMetasoundFrontendLiteral& InLiteral);
+	UE_API void SetSourceValue(const TArray<bool>& InSourceValue);
+};
+
+/**
+ * Viewmodel for binding an integer-type MetaSound Literal with a widget.
+ */
+UCLASS(BlueprintType, EditInlineNew, DefaultToInstanced, CollapseCategories, DisplayName = "MetaSound Literal Integer Viewmodel")
+class UMetaSoundLiteralViewModel_Integer : public UMVVMViewModelBase
+{
+	GENERATED_BODY()
+
+protected:
+	// The bound MetaSound Literal value. Should be bound using a two-way binding with the Literal variable of a MetaSoundInputViewModel.
+	UPROPERTY(BlueprintReadWrite, Transient, FieldNotify, Setter, Category = "Integer Literal", meta = (AllowPrivateAccess))
+	FMetasoundFrontendLiteral Literal;
+
+	// The value of the bound Literal as an integer.
+	UPROPERTY(BlueprintReadWrite, Transient, FieldNotify, Getter, Setter, Category = "Integer Literal", meta = (AllowPrivateAccess))
+	int32 SourceValue = 0;
+
+	//
+	UPROPERTY(BlueprintReadWrite, Transient, FieldNotify, Getter, Setter, Category = "Integer Literal", meta = (AllowPrivateAccess))
+	float NormalizedValue = 0.f;
+
+	//
+	UPROPERTY(BlueprintReadWrite, EditAnywhere, FieldNotify, Category = "Integer Literal", meta = (AllowPrivateAccess))
+	FInt32Range Range;
+
+public:
+	UFUNCTION(BlueprintCallable, FieldNotify, Category = "Integer Literal")
+	UE_API float GetStepSize() const;
+
+	const FMetasoundFrontendLiteral& GetLiteral() const { return Literal; }
+	int32 GetSourceValue() const { return SourceValue; }
+	float GetNormalizedValue() const { return NormalizedValue; }
+
+	UE_API void SetLiteral(const FMetasoundFrontendLiteral& InLiteral);
+	UE_API void SetSourceValue(int32 InSourceValue);
+	UE_API void SetNormalizedValue(float InNormalizedValue);
+
+private:
+	bool bIsUpdating = false;
+};
+
+/**
+ * Viewmodel for binding an integer array-type MetaSound Literal with a widget.
+ */
+UCLASS(BlueprintType, EditInlineNew, DefaultToInstanced, CollapseCategories, DisplayName = "MetaSound Literal Integer Array Viewmodel")
+class UMetaSoundLiteralViewModel_IntegerArray : public UMVVMViewModelBase
+{
+	GENERATED_BODY()
+
+protected:
+	// The bound MetaSound Literal value. Should be bound using a two-way binding with the Literal variable of a MetaSoundInputViewModel.
+	UPROPERTY(BlueprintReadWrite, Transient, FieldNotify, Setter, Category = "Integer Array Literal", meta = (AllowPrivateAccess))
+	FMetasoundFrontendLiteral Literal;
+
+	// The value of the bound Literal as an integer array.
+	UPROPERTY(BlueprintReadWrite, Transient, FieldNotify, Getter, Setter, Category = "Integer Array Literal", meta = (AllowPrivateAccess))
+	TArray<int32> SourceValue;
+
+public:
+	const FMetasoundFrontendLiteral& GetLiteral() const { return Literal; }
+	TArray<int32> GetSourceValue() const { return SourceValue; }
+
+	UE_API void SetLiteral(const FMetasoundFrontendLiteral& InLiteral);
+	UE_API void SetSourceValue(const TArray<int32>& InSourceValue);
+};
+
 /**
  * Viewmodel for converting a float-type MetaSound Literal between Source, Normalized, and Display values.
  *
@@ -38,22 +156,22 @@ protected:
 	bool bShowDisplayValues = true;
 
 	// The bound MetaSound Literal value. Should be bound using a two-way binding with the Literal variable of a MetaSoundInputViewModel.
-	UPROPERTY(BlueprintReadWrite, FieldNotify, Setter, Category = "Float Literal", meta = (AllowPrivateAccess))
+	UPROPERTY(BlueprintReadWrite, Transient, FieldNotify, Setter, Category = "Float Literal", meta = (AllowPrivateAccess))
 	FMetasoundFrontendLiteral Literal;
 
 	// The value of the bound Literal as a float.
-	UPROPERTY(BlueprintReadWrite, FieldNotify, Getter, Setter, Category = "Float Literal", meta = (AllowPrivateAccess))
+	UPROPERTY(BlueprintReadWrite, Transient, FieldNotify, Getter, Setter, Category = "Float Literal", meta = (AllowPrivateAccess))
 	float SourceValue = 0.f;
 
 	/**
 	 * The normalized (0-1) representation of the value currently surfaced to the UI. If bShowDisplayValues is true, this is DisplayValue mapped into
 	 * 0-1 space. Otherwise, it's SourceValue mapped into 0-1 space. Used by knobs/sliders purely for visual feedback.
 	 */
-	UPROPERTY(BlueprintReadWrite, FieldNotify, Getter, Setter, Category = "Float Literal", meta = (AllowPrivateAccess))
+	UPROPERTY(BlueprintReadWrite, Transient, FieldNotify, Getter, Setter, Category = "Float Literal", meta = (AllowPrivateAccess))
 	float NormalizedValue = 0.f;
 
 	// The user-facing display value (e.g. dB, Hz, semitones). Can be modified and is automatically converted to the expected SourceRange.
-	UPROPERTY(BlueprintReadWrite, FieldNotify, Getter, Setter, Category = "Float Literal", meta = (AllowPrivateAccess))
+	UPROPERTY(BlueprintReadWrite, Transient, FieldNotify, Getter, Setter, Category = "Float Literal", meta = (AllowPrivateAccess))
 	float DisplayValue = 0.f;
 
 	// Mapping pair used to convert between float literal, normalized, and display values.
@@ -63,42 +181,169 @@ protected:
 public:
 	// Returns the range minimum of the SourceRange.
 	UFUNCTION(BlueprintCallable, BlueprintPure, FieldNotify, Category = "TechAudioTools")
-	TECHAUDIOTOOLSMETASOUND_API float GetSourceRangeMin() const;
+	UE_API float GetSourceRangeMin() const;
 
 	// Returns the range maximum of the SourceRange.
 	UFUNCTION(BlueprintCallable, BlueprintPure, FieldNotify, Category = "TechAudioTools")
-	TECHAUDIOTOOLSMETASOUND_API float GetSourceRangeMax() const;
+	UE_API float GetSourceRangeMax() const;
 
 	// Returns the units associated with the SourceValue.
 	UFUNCTION(BlueprintCallable, BlueprintPure, FieldNotify, Category = "TechAudioTools")
-	TECHAUDIOTOOLSMETASOUND_API ETechAudioToolsFloatUnit GetSourceUnits() const;
+	UE_API ETechAudioToolsFloatUnit GetSourceUnits() const;
 
 	// Returns the range minimum of the DisplayRange.
 	UFUNCTION(BlueprintCallable, BlueprintPure, FieldNotify, Category = "TechAudioTools")
-	TECHAUDIOTOOLSMETASOUND_API float GetDisplayRangeMin() const;
+	UE_API float GetDisplayRangeMin() const;
 
 	// Returns the range maximum of the DisplayRange.
 	UFUNCTION(BlueprintCallable, BlueprintPure, FieldNotify, Category = "TechAudioTools")
-	TECHAUDIOTOOLSMETASOUND_API float GetDisplayRangeMax() const;
+	UE_API float GetDisplayRangeMax() const;
 
 	// Returns the units associated with the DisplayValue.
 	UFUNCTION(BlueprintCallable, BlueprintPure, FieldNotify, Category = "TechAudioTools")
-	TECHAUDIOTOOLSMETASOUND_API ETechAudioToolsFloatUnit GetDisplayUnits() const;
+	UE_API ETechAudioToolsFloatUnit GetDisplayUnits() const;
 
-	explicit UMetaSoundLiteralViewModel_Float(const FObjectInitializer& ObjectInitializer);
+	UE_API explicit UMetaSoundLiteralViewModel_Float(const FObjectInitializer& ObjectInitializer);
 
 	bool IsShowingDisplayValues() const { return bShowDisplayValues; }
-	FMetasoundFrontendLiteral GetLiteral() const { return Literal; }
+	const FMetasoundFrontendLiteral& GetLiteral() const { return Literal; }
 	float GetSourceValue() const { return SourceValue; }
 	float GetNormalizedValue() const { return NormalizedValue; }
 	float GetDisplayValue() const { return DisplayValue; }
 
-	TECHAUDIOTOOLSMETASOUND_API void SetShowDisplayValues(bool bInIsShowingDisplayValues);
-	TECHAUDIOTOOLSMETASOUND_API void SetLiteral(const FMetasoundFrontendLiteral& InLiteral);
-	TECHAUDIOTOOLSMETASOUND_API void SetSourceValue(float InSourceValue);
-	TECHAUDIOTOOLSMETASOUND_API void SetNormalizedValue(float InNormalizedValue);
-	TECHAUDIOTOOLSMETASOUND_API void SetDisplayValue(float InDisplayValue);
+	UE_API void SetShowDisplayValues(bool bInIsShowingDisplayValues);
+	UE_API void SetLiteral(const FMetasoundFrontendLiteral& InLiteral);
+	UE_API void SetSourceValue(float InSourceValue);
+	UE_API void SetNormalizedValue(float InNormalizedValue);
+	UE_API void SetDisplayValue(float InDisplayValue);
 
 private:
 	bool bIsUpdating = false;
 };
+
+/**
+ * Viewmodel for binding a float array-type MetaSound Literal with a widget.
+ */
+UCLASS(BlueprintType, EditInlineNew, DefaultToInstanced, CollapseCategories, DisplayName = "MetaSound Literal Float Array Viewmodel")
+class UMetaSoundLiteralViewModel_FloatArray : public UMVVMViewModelBase
+{
+	GENERATED_BODY()
+
+protected:
+	// The bound MetaSound Literal value. Should be bound using a two-way binding with the Literal variable of a MetaSoundInputViewModel.
+	UPROPERTY(BlueprintReadWrite, Transient, FieldNotify, Setter, Category = "Float Array Literal", meta = (AllowPrivateAccess))
+	FMetasoundFrontendLiteral Literal;
+
+	// The value of the bound Literal as a float array.
+	UPROPERTY(BlueprintReadWrite, Transient, FieldNotify, Getter, Setter, Category = "Float Array Literal", meta = (AllowPrivateAccess))
+	TArray<float> SourceValue;
+
+public:
+	const FMetasoundFrontendLiteral& GetLiteral() const { return Literal; }
+	TArray<float> GetSourceValue() const { return SourceValue; }
+
+	UE_API void SetLiteral(const FMetasoundFrontendLiteral& InLiteral);
+	UE_API void SetSourceValue(const TArray<float>& InSourceValue);
+};
+
+/**
+ * Viewmodel for binding a string-type MetaSound Literal with a widget.
+ */
+UCLASS(BlueprintType, EditInlineNew, DefaultToInstanced, CollapseCategories, DisplayName = "MetaSound Literal String Viewmodel")
+class UMetaSoundLiteralViewModel_String : public UMVVMViewModelBase
+{
+	GENERATED_BODY()
+
+protected:
+	// The bound MetaSound Literal value. Should be bound using a two-way binding with the Literal variable of a MetaSoundInputViewModel.
+	UPROPERTY(BlueprintReadWrite, Transient, FieldNotify, Setter, Category = "String Literal", meta = (AllowPrivateAccess))
+	FMetasoundFrontendLiteral Literal;
+
+	// The value of the bound Literal as a bool.
+	UPROPERTY(BlueprintReadWrite, Transient, FieldNotify, Getter, Setter, Category = "String Literal", meta = (AllowPrivateAccess))
+	FString SourceValue;
+
+public:
+	const FMetasoundFrontendLiteral& GetLiteral() const { return Literal; }
+	FString GetSourceValue() const { return SourceValue; }
+
+	UE_API void SetLiteral(const FMetasoundFrontendLiteral& InLiteral);
+	UE_API void SetSourceValue(const FString& InSourceValue);
+};
+
+/**
+ * Viewmodel for binding a string array-type MetaSound Literal with a widget.
+ */
+UCLASS(BlueprintType, EditInlineNew, DefaultToInstanced, CollapseCategories, DisplayName = "MetaSound Literal String Array Viewmodel")
+class UMetaSoundLiteralViewModel_StringArray : public UMVVMViewModelBase
+{
+	GENERATED_BODY()
+
+protected:
+	// The bound MetaSound Literal value. Should be bound using a two-way binding with the Literal variable of a MetaSoundInputViewModel.
+	UPROPERTY(BlueprintReadWrite, Transient, FieldNotify, Setter, Category = "String Array Literal", meta = (AllowPrivateAccess))
+	FMetasoundFrontendLiteral Literal;
+
+	// The value of the bound Literal as a string array.
+	UPROPERTY(BlueprintReadWrite, Transient, FieldNotify, Getter, Setter, Category = "String Array Literal", meta = (AllowPrivateAccess))
+	TArray<FString> SourceValue;
+
+public:
+	const FMetasoundFrontendLiteral& GetLiteral() const { return Literal; }
+	TArray<FString> GetSourceValue() const { return SourceValue; }
+
+	UE_API void SetLiteral(const FMetasoundFrontendLiteral& InLiteral);
+	UE_API void SetSourceValue(const TArray<FString>& InSourceValue);
+};
+
+/**
+ * Viewmodel for binding an object-type MetaSound Literal with a widget.
+ */
+UCLASS(BlueprintType, EditInlineNew, DefaultToInstanced, CollapseCategories, DisplayName = "MetaSound Literal Object Viewmodel")
+class UMetaSoundLiteralViewModel_Object : public UMVVMViewModelBase
+{
+	GENERATED_BODY()
+
+protected:
+	// The bound MetaSound Literal value. Should be bound using a two-way binding with the Literal variable of a MetaSoundInputViewModel.
+	UPROPERTY(BlueprintReadWrite, Transient, FieldNotify, Setter, Category = "Object Literal", meta = (AllowPrivateAccess))
+	FMetasoundFrontendLiteral Literal;
+
+	// The value of the bound Literal as a bool.
+	UPROPERTY(BlueprintReadWrite, Transient, FieldNotify, Getter, Setter, Category = "Object Literal", meta = (AllowPrivateAccess))
+	TObjectPtr<UObject> SourceValue;
+
+public:
+	const FMetasoundFrontendLiteral& GetLiteral() const { return Literal; }
+	UObject* GetSourceValue() const { return SourceValue.Get(); }
+
+	UE_API void SetLiteral(const FMetasoundFrontendLiteral& InLiteral);
+	UE_API void SetSourceValue(UObject* InSourceValue);
+};
+
+/**
+ * Viewmodel for binding an object array-type MetaSound Literal with a widget.
+ */
+UCLASS(BlueprintType, EditInlineNew, DefaultToInstanced, CollapseCategories, DisplayName = "MetaSound Literal Object Array Viewmodel")
+class UMetaSoundLiteralViewModel_ObjectArray : public UMVVMViewModelBase
+{
+	GENERATED_BODY()
+
+protected:
+	// The bound MetaSound Literal value. Should be bound using a two-way binding with the Literal variable of a MetaSoundInputViewModel.
+	UPROPERTY(BlueprintReadWrite, Transient, FieldNotify, Setter, Category = "Object Array Literal", meta = (AllowPrivateAccess))
+	FMetasoundFrontendLiteral Literal;
+
+	// The value of the bound Literal as an object array.
+	UPROPERTY(BlueprintReadWrite, Transient, FieldNotify, Getter, Setter, Category = "Object Array Literal", meta = (AllowPrivateAccess))
+	TArray<TObjectPtr<UObject>> SourceValue;
+
+public:
+	const FMetasoundFrontendLiteral& GetLiteral() const { return Literal; }
+	TArray<TObjectPtr<UObject>> GetSourceValue() const { return SourceValue; }
+
+	UE_API void SetLiteral(const FMetasoundFrontendLiteral& InLiteral);
+	UE_API void SetSourceValue(const TArray<TObjectPtr<UObject>>& InSourceValue);
+};
+
+#undef UE_API
