TechAudioTools - fixed documentation errors

#rnx

diff --git a/Engine/Plugins/Experimental/TechAudioTools/Source/TechAudioToolsMetaSound/Public/ViewModels/MetaSoundLiteralViewModel.h b/Engine/Plugins/Experimental/TechAudioTools/Source/TechAudioToolsMetaSound/Public/ViewModels/MetaSoundLiteralViewModel.h
--- a/Engine/Plugins/Experimental/TechAudioTools/Source/TechAudioToolsMetaSound/Public/ViewModels/MetaSoundLiteralViewModel.h
+++ b/Engine/Plugins/Experimental/TechAudioTools/Source/TechAudioToolsMetaSound/Public/ViewModels/MetaSoundLiteralViewModel.h
@@ -80,11 +80,11 @@ protected:
 	UPROPERTY(BlueprintReadWrite, Transient, FieldNotify, Getter, Setter, Category = "Integer Literal", meta = (AllowPrivateAccess))
 	int32 SourceValue = 0;
 
-	//
+	// The SourceValue normalized to a 0-1 range based on the set Range values. Used for UI like sliders and knobs that use a normalized range.
 	UPROPERTY(BlueprintReadWrite, Transient, FieldNotify, Getter, Setter, Category = "Integer Literal", meta = (AllowPrivateAccess))
 	float NormalizedValue = 0.f;
 
-	//
+	// Sets the min and max values of the UI range. Range bounds should be closed (either Inclusive or Exclusive) and not open.
 	UPROPERTY(BlueprintReadWrite, EditAnywhere, FieldNotify, Category = "Integer Literal", meta = (AllowPrivateAccess))
 	FInt32Range Range;
 
@@ -259,7 +259,7 @@ protected:
 	UPROPERTY(BlueprintReadWrite, Transient, FieldNotify, Setter, Category = "String Literal", meta = (AllowPrivateAccess))
 	FMetasoundFrontendLiteral Literal;
 
-	// The value of the bound Literal as a bool.
+	// The value of the bound Literal as a string.
 	UPROPERTY(BlueprintReadWrite, Transient, FieldNotify, Getter, Setter, Category = "String Literal", meta = (AllowPrivateAccess))
 	FString SourceValue;
 
@@ -309,7 +309,7 @@ protected:
 	UPROPERTY(BlueprintReadWrite, Transient, FieldNotify, Setter, Category = "Object Literal", meta = (AllowPrivateAccess))
 	FMetasoundFrontendLiteral Literal;
 
-	// The value of the bound Literal as a bool.
+	// The value of the bound Literal as an object.
 	UPROPERTY(BlueprintReadWrite, Transient, FieldNotify, Getter, Setter, Category = "Object Literal", meta = (AllowPrivateAccess))
 	TObjectPtr<UObject> SourceValue;
 
