TechAudioTools - Added AudioComponentViewModel

#rnx

diff --git a/Engine/Plugins/Experimental/TechAudioTools/Source/TechAudioTools/Private/Viewmodels/AudioComponentViewModel.cpp b/Engine/Plugins/Experimental/TechAudioTools/Source/TechAudioTools/Private/Viewmodels/AudioComponentViewModel.cpp
new file mode 100644
--- /dev/null
+++ b/Engine/Plugins/Experimental/TechAudioTools/Source/TechAudioTools/Private/Viewmodels/AudioComponentViewModel.cpp
@@ -0,0 +1,86 @@
+﻿// Copyright Epic Games, Inc. All Rights Reserved.
+
+#include "Viewmodels/AudioComponentViewModel.h"
+
+#include UE_INLINE_GENERATED_CPP_BY_NAME(AudioComponentViewModel)
+
+void UAudioComponentViewModel::SetAudioComponent(UAudioComponent* InAudioComponent)
+{
+	if (AudioComponent == InAudioComponent)
+	{
+		return;
+	}
+
+	UnbindDelegates();
+	AudioComponent = InAudioComponent;
+
+	if (InAudioComponent)
+	{
+		SetPlayState(InAudioComponent->GetPlayState());
+		UE_MVVM_SET_PROPERTY_VALUE(bIsVirtualized, InAudioComponent->IsVirtualized());
+		BindDelegates();
+	}
+	else
+	{
+		UE_MVVM_SET_PROPERTY_VALUE(PlayState, EAudioComponentPlayState::Stopped);
+		UE_MVVM_SET_PROPERTY_VALUE(bIsPlaying, false);
+		UE_MVVM_SET_PROPERTY_VALUE(bIsStopped, true);
+		UE_MVVM_SET_PROPERTY_VALUE(bIsFadingIn, false);
+		UE_MVVM_SET_PROPERTY_VALUE(bIsFadingOut, false);
+		UE_MVVM_SET_PROPERTY_VALUE(bIsVirtualized, false);
+	}
+}
+
+void UAudioComponentViewModel::BeginDestroy()
+{
+	Super::BeginDestroy();
+	UnbindDelegates();
+}
+
+void UAudioComponentViewModel::BindDelegates()
+{
+	if (AudioComponent.IsValid())
+	{
+		AudioComponent->OnAudioPlayStateChanged.AddUniqueDynamic(this, &ThisClass::SetPlayState);
+		AudioComponent->OnAudioFinished.AddUniqueDynamic(this, &ThisClass::OnAudioFinished);
+		AudioComponent->OnAudioVirtualizationChanged.AddUniqueDynamic(this, &ThisClass::OnVirtualizationChanged);
+	}
+}
+
+void UAudioComponentViewModel::UnbindDelegates()
+{
+	if (AudioComponent.IsValid())
+	{
+		AudioComponent->OnAudioPlayStateChanged.RemoveDynamic(this, &ThisClass::SetPlayState);
+		AudioComponent->OnAudioFinished.RemoveDynamic(this, &ThisClass::OnAudioFinished);
+		AudioComponent->OnAudioVirtualizationChanged.RemoveDynamic(this, &ThisClass::OnVirtualizationChanged);
+	}
+}
+
+void UAudioComponentViewModel::OnAudioFinished()
+{
+	if (AudioComponent.IsValid() && AudioComponent->bAutoDestroy != 0)
+	{
+		SetAudioComponent(nullptr);
+	}
+	else
+	{
+		SetPlayState(EAudioComponentPlayState::Stopped);
+	}
+}
+
+void UAudioComponentViewModel::OnVirtualizationChanged(const bool bInIsVirtualized)
+{
+	UE_MVVM_SET_PROPERTY_VALUE(bIsVirtualized, bInIsVirtualized);
+}
+
+void UAudioComponentViewModel::SetPlayState(const EAudioComponentPlayState NewPlayState)
+{
+	if (UE_MVVM_SET_PROPERTY_VALUE(PlayState, NewPlayState))
+	{
+		UE_MVVM_SET_PROPERTY_VALUE(bIsPlaying, NewPlayState == EAudioComponentPlayState::Playing);
+		UE_MVVM_SET_PROPERTY_VALUE(bIsStopped, NewPlayState == EAudioComponentPlayState::Stopped);
+		UE_MVVM_SET_PROPERTY_VALUE(bIsFadingIn, NewPlayState == EAudioComponentPlayState::FadingIn);
+		UE_MVVM_SET_PROPERTY_VALUE(bIsFadingOut, NewPlayState == EAudioComponentPlayState::FadingOut);
+	}
+}
diff --git a/Engine/Plugins/Experimental/TechAudioTools/Source/TechAudioTools/Public/Viewmodels/AudioComponentViewModel.h b/Engine/Plugins/Experimental/TechAudioTools/Source/TechAudioTools/Public/Viewmodels/AudioComponentViewModel.h
new file mode 100644
--- /dev/null
+++ b/Engine/Plugins/Experimental/TechAudioTools/Source/TechAudioTools/Public/Viewmodels/AudioComponentViewModel.h
@@ -0,0 +1,74 @@
+﻿// Copyright Epic Games, Inc. All Rights Reserved.
+
+#pragma once
+
+#include "MVVMViewModelBase.h"
+#include "Components/AudioComponent.h"
+
+#include "AudioComponentViewModel.generated.h"
+
+/**
+ * Viewmodel for binding audio component state to widgets.
+ */
+UCLASS(MinimalAPI, BlueprintType, DisplayName = "Audio Component Viewmodel")
+class UAudioComponentViewModel : public UMVVMViewModelBase
+{
+	GENERATED_BODY()
+
+private:
+	// Returns the enumerated play state of the audio component.
+	UPROPERTY(BlueprintReadOnly, Transient, FieldNotify, Getter = "GetPlayState", Category = "Audio Component", meta = (AllowPrivateAccess))
+	EAudioComponentPlayState PlayState = EAudioComponentPlayState::Stopped;
+
+	// True if the audio component is playing.
+	UPROPERTY(BlueprintReadOnly, Transient, FieldNotify, Getter = "IsPlaying", Category = "Audio Component", meta = (AllowPrivateAccess))
+	bool bIsPlaying = false;
+
+	// True if the audio component is stopped.
+	UPROPERTY(BlueprintReadOnly, Transient, FieldNotify, Getter = "IsStopped", Category = "Audio Component", meta = (AllowPrivateAccess))
+	bool bIsStopped = true;
+
+	// True if the audio component is fading in.
+	UPROPERTY(BlueprintReadOnly, Transient, FieldNotify, Getter = "IsFadingIn", Category = "Audio Component", meta = (AllowPrivateAccess))
+	bool bIsFadingIn = false;
+
+	// True if the audio component is fading out.
+	UPROPERTY(BlueprintReadOnly, Transient, FieldNotify, Getter = "IsFadingOut", Category = "Audio Component", meta = (AllowPrivateAccess))
+	bool bIsFadingOut = false;
+
+	// True if the audio component is virtualized.
+	UPROPERTY(BlueprintReadOnly, Transient, FieldNotify, Getter = "IsVirtualized", Category = "Audio Component", meta = (AllowPrivateAccess))
+	bool bIsVirtualized = false;
+
+protected:
+	UPROPERTY(Transient)
+	TWeakObjectPtr<UAudioComponent> AudioComponent;
+
+public:
+	// Sets the audio component this viewmodel should bind to.
+	UFUNCTION(BlueprintCallable, Category = "TechAudioTools|Audio Component")
+	TECHAUDIOTOOLS_API virtual void SetAudioComponent(UAudioComponent* InAudioComponent);
+
+	EAudioComponentPlayState GetPlayState() const { return PlayState; }
+	bool IsPlaying() const { return bIsPlaying; }
+	bool IsStopped() const { return bIsStopped; }
+	bool IsFadingIn() const { return bIsFadingIn; }
+	bool IsFadingOut() const { return bIsFadingOut; }
+	bool IsVirtualized() const { return bIsVirtualized; }
+
+	virtual void BeginDestroy() override;
+
+protected:
+	void BindDelegates();
+	void UnbindDelegates();
+
+	UFUNCTION()
+	void OnAudioFinished();
+
+	UFUNCTION()
+	void OnVirtualizationChanged(bool bInIsVirtualized);
+
+private:
+	UFUNCTION()
+	void SetPlayState(EAudioComponentPlayState NewPlayState);
+};
diff --git a/Engine/Plugins/Experimental/TechAudioTools/Source/TechAudioTools/TechAudioTools.Build.cs b/Engine/Plugins/Experimental/TechAudioTools/Source/TechAudioTools/TechAudioTools.Build.cs
--- a/Engine/Plugins/Experimental/TechAudioTools/Source/TechAudioTools/TechAudioTools.Build.cs
+++ b/Engine/Plugins/Experimental/TechAudioTools/Source/TechAudioTools/TechAudioTools.Build.cs
@@ -12,6 +12,7 @@ public class TechAudioTools : ModuleRules
             new string[]
             {
                 "Core",
+                "ModelViewViewModel",
             }
         );
 
