TechAudioTools - added BandwidthOct and Tempo as new float unit types for label formatting

#rnx

diff --git a/Engine/Plugins/Experimental/TechAudioTools/Source/TechAudioTools/Public/TechAudioToolsTypes.h b/Engine/Plugins/Experimental/TechAudioTools/Source/TechAudioTools/Public/TechAudioToolsTypes.h
--- a/Engine/Plugins/Experimental/TechAudioTools/Source/TechAudioTools/Public/TechAudioToolsTypes.h
+++ b/Engine/Plugins/Experimental/TechAudioTools/Source/TechAudioTools/Public/TechAudioToolsTypes.h
@@ -25,6 +25,8 @@ enum class ETechAudioToolsFloatUnit : uint8
 	TimeMilliseconds	UMETA(DisplayName = "Milliseconds", 		ToolTip = "Time in ms"),
 	Degrees 			UMETA(DisplayName = "Degrees",  			ToolTip = "Angle in degrees (0-360)"),
 	Pan 				UMETA(DisplayName = "Pan",  				ToolTip = "Pan value: negative = L, 0 = C, positive = R"),
+	BandwidthOct		UMETA(DisplayName = "Bandwidth",			ToolTip = "Filter bandwidth in octaves (1.0 = one octave)"),
+	Tempo   			UMETA(DisplayName = "Tempo",				ToolTip = "Tempo in beats per minute (BPM)")
 };
 
 // The mapping logic applied between source and display values.
