TechAudioTools - Remove clamp when converting between source and display values while using Default mapping type

#rnx

diff --git a/Engine/Plugins/Experimental/TechAudioTools/Source/TechAudioTools/Private/TechAudioToolsFloatMapping.cpp b/Engine/Plugins/Experimental/TechAudioTools/Source/TechAudioTools/Private/TechAudioToolsFloatMapping.cpp
--- a/Engine/Plugins/Experimental/TechAudioTools/Source/TechAudioTools/Private/TechAudioToolsFloatMapping.cpp
+++ b/Engine/Plugins/Experimental/TechAudioTools/Source/TechAudioTools/Private/TechAudioToolsFloatMapping.cpp
@@ -73,7 +73,7 @@ float UTechAudioToolsFloatMapping::SourceToDisplay(const float InSource) const
 	switch (MappingType)
 	{
 	case ETechAudioToolsFloatMappingType::Default:
-		return FMath::Clamp(InSource, Range.Min, Range.Max);
+		return InSource;
 
 	case ETechAudioToolsFloatMappingType::MapRange:
 		{
@@ -97,7 +97,7 @@ float UTechAudioToolsFloatMapping::DisplayToSource(const float InDisplay) const
 	switch (MappingType)
 	{
 	case ETechAudioToolsFloatMappingType::Default:
-		return FMath::Clamp(InDisplay, Range.Min, Range.Max);
+		return InDisplay;
 
 	case ETechAudioToolsFloatMappingType::MapRange:
 		{
