TechAudioTools - Fix MetaSound Editor Viewmodel TMap lookups

#Synced-CL 50450250

diff --git a/Engine/Plugins/Experimental/TechAudioTools/Source/TechAudioToolsMetaSoundEditor/Private/ViewModels/MetaSoundEditorViewModel.cpp b/Engine/Plugins/Experimental/TechAudioTools/Source/TechAudioToolsMetaSoundEditor/Private/ViewModels/MetaSoundEditorViewModel.cpp
--- a/Engine/Plugins/Experimental/TechAudioTools/Source/TechAudioToolsMetaSoundEditor/Private/ViewModels/MetaSoundEditorViewModel.cpp
+++ b/Engine/Plugins/Experimental/TechAudioTools/Source/TechAudioToolsMetaSoundEditor/Private/ViewModels/MetaSoundEditorViewModel.cpp
@@ -208,7 +208,7 @@ void UMetaSoundEditorViewModel::OnGraphInputDefaultChanged(const FGraphInputDefa
 
 void UMetaSoundEditorViewModel::OnInputDisplayNameChanged(const FName InputName, const FText InDisplayName)
 {
-	if (UMetaSoundInputEditorViewModel* InputEditorViewModel = Cast<UMetaSoundInputEditorViewModel>(*InputViewModels.Find(InputName)))
+	if (UMetaSoundInputEditorViewModel* InputEditorViewModel = Cast<UMetaSoundInputEditorViewModel>(InputViewModels.FindRef(InputName)))
 	{
 		InputEditorViewModel->OnInputDisplayNameChanged(InDisplayName);
 	}
@@ -216,7 +216,7 @@ void UMetaSoundEditorViewModel::OnInputDisplayNameChanged(const FName InputName,
 
 void UMetaSoundEditorViewModel::OnInputDescriptionChanged(const FName InputName, FText InDescription)
 {
-	if (UMetaSoundInputEditorViewModel* InputEditorViewModel = Cast<UMetaSoundInputEditorViewModel>(*InputViewModels.Find(InputName)))
+	if (UMetaSoundInputEditorViewModel* InputEditorViewModel = Cast<UMetaSoundInputEditorViewModel>(InputViewModels.FindRef(InputName)))
 	{
 		InputEditorViewModel->OnInputDescriptionChanged(InDescription);
 	}
@@ -224,7 +224,7 @@ void UMetaSoundEditorViewModel::OnInputDescriptionChanged(const FName InputName,
 
 void UMetaSoundEditorViewModel::OnInputSortOrderIndexChanged(const FName InputName, const int32 InSortOrderIndex)
 {
-	if (UMetaSoundInputEditorViewModel* InputEditorViewModel = Cast<UMetaSoundInputEditorViewModel>(*InputViewModels.Find(InputName)))
+	if (UMetaSoundInputEditorViewModel* InputEditorViewModel = Cast<UMetaSoundInputEditorViewModel>(InputViewModels.FindRef(InputName)))
 	{
 		InputEditorViewModel->OnInputSortOrderIndexChanged(InSortOrderIndex);
 	}
@@ -232,7 +232,7 @@ void UMetaSoundEditorViewModel::OnInputSortOrderIndexChanged(const FName InputNa
 
 void UMetaSoundEditorViewModel::OnInputIsAdvancedDisplayChanged(const FName InputName, const bool bInIsAdvancedDisplay)
 {
-	if (UMetaSoundInputEditorViewModel* InputEditorViewModel = Cast<UMetaSoundInputEditorViewModel>(*InputViewModels.Find(InputName)))
+	if (UMetaSoundInputEditorViewModel* InputEditorViewModel = Cast<UMetaSoundInputEditorViewModel>(InputViewModels.FindRef(InputName)))
 	{
 		InputEditorViewModel->OnInputIsAdvancedDisplayChanged(bInIsAdvancedDisplay);
 	}
@@ -240,7 +240,7 @@ void UMetaSoundEditorViewModel::OnInputIsAdvancedDisplayChanged(const FName Inpu
 
 void UMetaSoundEditorViewModel::OnOutputDisplayNameChanged(const FName OutputName, const FText NewName)
 {
-	if (UMetaSoundOutputEditorViewModel* OutputEditorViewModel = Cast<UMetaSoundOutputEditorViewModel>(OutputViewModels.FindAndRemoveChecked(OutputName)))
+	if (UMetaSoundOutputEditorViewModel* OutputEditorViewModel = Cast<UMetaSoundOutputEditorViewModel>(OutputViewModels.FindRef(OutputName)))
 	{
 		OutputEditorViewModel->OnOutputDisplayNameChanged(NewName);
 	}
@@ -248,7 +248,7 @@ void UMetaSoundEditorViewModel::OnOutputDisplayNameChanged(const FName OutputNam
 
 void UMetaSoundEditorViewModel::OnOutputDescriptionChanged(const FName OutputName, const FText InDescription)
 {
-	if (UMetaSoundOutputEditorViewModel* OutputEditorViewModel = Cast<UMetaSoundOutputEditorViewModel>(OutputViewModels.FindAndRemoveChecked(OutputName)))
+	if (UMetaSoundOutputEditorViewModel* OutputEditorViewModel = Cast<UMetaSoundOutputEditorViewModel>(OutputViewModels.FindRef(OutputName)))
 	{
 		OutputEditorViewModel->OnOutputDescriptionChanged(InDescription);
 	}
@@ -256,7 +256,7 @@ void UMetaSoundEditorViewModel::OnOutputDescriptionChanged(const FName OutputNam
 
 void UMetaSoundEditorViewModel::OnOutputSortOrderIndexChanged(const FName OutputName, const int32 InSortOrderIndex)
 {
-	if (UMetaSoundOutputEditorViewModel* OutputEditorViewModel = Cast<UMetaSoundOutputEditorViewModel>(OutputViewModels.FindAndRemoveChecked(OutputName)))
+	if (UMetaSoundOutputEditorViewModel* OutputEditorViewModel = Cast<UMetaSoundOutputEditorViewModel>(OutputViewModels.FindRef(OutputName)))
 	{
 		OutputEditorViewModel->OnOutputSortOrderIndexChanged(InSortOrderIndex);
 	}
@@ -264,7 +264,7 @@ void UMetaSoundEditorViewModel::OnOutputSortOrderIndexChanged(const FName Output
 
 void UMetaSoundEditorViewModel::OnOutputIsAdvancedDisplayChanged(const FName OutputName, const bool bInIsAdvancedDisplay)
 {
-	if (UMetaSoundOutputEditorViewModel* OutputEditorViewModel = Cast<UMetaSoundOutputEditorViewModel>(OutputViewModels.FindAndRemoveChecked(OutputName)))
+	if (UMetaSoundOutputEditorViewModel* OutputEditorViewModel = Cast<UMetaSoundOutputEditorViewModel>(OutputViewModels.FindRef(OutputName)))
 	{
 		OutputEditorViewModel->OnOutputIsAdvancedDisplayChanged(bInIsAdvancedDisplay);
 	}
