[TechAudioTools] Added support for transactions in MetaSound Literal Viewmodels

- Implemented PostEditUndo, which ensures bindings are executed when undoing and redoing transactions
- Add RF_Transactional flag in PostInitProperties() to ensure all literal viewmodels are transactable, regardless of whether bExposeInstanceInEditor is enabled in the viewmodel details

#rnx

#Synced-CL 52401088

diff --git a/Engine/Plugins/Experimental/TechAudioTools/Source/TechAudioToolsMetaSound/Private/ViewModels/MetaSoundLiteralViewModel.cpp b/Engine/Plugins/Experimental/TechAudioTools/Source/TechAudioToolsMetaSound/Private/ViewModels/MetaSoundLiteralViewModel.cpp
--- a/Engine/Plugins/Experimental/TechAudioTools/Source/TechAudioToolsMetaSound/Private/ViewModels/MetaSoundLiteralViewModel.cpp
+++ b/Engine/Plugins/Experimental/TechAudioTools/Source/TechAudioToolsMetaSound/Private/ViewModels/MetaSoundLiteralViewModel.cpp
@@ -6,6 +6,12 @@
 
 #include UE_INLINE_GENERATED_CPP_BY_NAME(MetaSoundLiteralViewModel)
 
+void UMetaSoundLiteralViewModel_Boolean::PostInitProperties()
+{
+	Super::PostInitProperties();
+	SetFlags(RF_Transactional);
+}
+
 void UMetaSoundLiteralViewModel_Boolean::SetLiteral(const FMetasoundFrontendLiteral& InLiteral)
 {
 	if (UE_MVVM_SET_PROPERTY_VALUE(Literal, InLiteral))
@@ -25,6 +31,20 @@ void UMetaSoundLiteralViewModel_Boolean::SetSourceValue(const bool InSourceValue
 	}
 }
 
+#if WITH_EDITOR
+void UMetaSoundLiteralViewModel_Boolean::PostEditUndo()
+{
+	Super::PostEditUndo();
+	SetSourceValue(SourceValue);
+}
+#endif
+
+void UMetaSoundLiteralViewModel_BooleanArray::PostInitProperties()
+{
+	Super::PostInitProperties();
+	SetFlags(RF_Transactional);
+}
+
 void UMetaSoundLiteralViewModel_BooleanArray::SetLiteral(const FMetasoundFrontendLiteral& InLiteral)
 {
 	if (UE_MVVM_SET_PROPERTY_VALUE(Literal, InLiteral))
@@ -44,6 +64,20 @@ void UMetaSoundLiteralViewModel_BooleanArray::SetSourceValue(const TArray<bool>&
 	}
 }
 
+#if WITH_EDITOR
+void UMetaSoundLiteralViewModel_BooleanArray::PostEditUndo()
+{
+	Super::PostEditUndo();
+	SetSourceValue(SourceValue);
+}
+#endif
+
+void UMetaSoundLiteralViewModel_Integer::PostInitProperties()
+{
+	Super::PostInitProperties();
+	SetFlags(RF_Transactional);
+}
+
 float UMetaSoundLiteralViewModel_Integer::GetStepSize() const
 {
 	if (!Range.HasLowerBound() || !Range.HasUpperBound())
@@ -173,6 +207,20 @@ void UMetaSoundLiteralViewModel_Integer::SetNormalizedValue(const float InNormal
 	}
 }
 
+#if WITH_EDITOR
+void UMetaSoundLiteralViewModel_Integer::PostEditUndo()
+{
+	Super::PostEditUndo();
+	SetSourceValue(SourceValue);
+}
+#endif
+
+void UMetaSoundLiteralViewModel_IntegerArray::PostInitProperties()
+{
+	Super::PostInitProperties();
+	SetFlags(RF_Transactional);
+}
+
 void UMetaSoundLiteralViewModel_IntegerArray::SetLiteral(const FMetasoundFrontendLiteral& InLiteral)
 {
 	if (UE_MVVM_SET_PROPERTY_VALUE(Literal, InLiteral))
@@ -192,6 +240,14 @@ void UMetaSoundLiteralViewModel_IntegerArray::SetSourceValue(const TArray<int32>
 	}
 }
 
+#if WITH_EDITOR
+void UMetaSoundLiteralViewModel_IntegerArray::PostEditUndo()
+{
+	Super::PostEditUndo();
+	SetSourceValue(SourceValue);
+}
+#endif
+
 float UMetaSoundLiteralViewModel_Float::GetSourceRangeMin() const
 {
 	return RangeValues ? RangeValues->GetSourceMin() : 0.f;
@@ -227,6 +283,12 @@ UMetaSoundLiteralViewModel_Float::UMetaSoundLiteralViewModel_Float(const FObject
 	RangeValues = ObjectInitializer.CreateDefaultSubobject<UTechAudioToolsFloatMapping>(this, TEXT("RangeValues"));
 }
 
+void UMetaSoundLiteralViewModel_Float::PostInitProperties()
+{
+	Super::PostInitProperties();
+	SetFlags(RF_Transactional);
+}
+
 void UMetaSoundLiteralViewModel_Float::SetShowDisplayValues(const bool bInIsShowingDisplayValues)
 {
 	if (!bIsUpdating && RangeValues)
@@ -329,6 +391,20 @@ void UMetaSoundLiteralViewModel_Float::SetDisplayValue(const float InDisplayValu
 	}
 }
 
+#if WITH_EDITOR
+void UMetaSoundLiteralViewModel_Float::PostEditUndo()
+{
+	Super::PostEditUndo();
+	SetSourceValue(SourceValue);
+}
+#endif
+
+void UMetaSoundLiteralViewModel_FloatArray::PostInitProperties()
+{
+	Super::PostInitProperties();
+	SetFlags(RF_Transactional);
+}
+
 void UMetaSoundLiteralViewModel_FloatArray::SetLiteral(const FMetasoundFrontendLiteral& InLiteral)
 {
 	if (UE_MVVM_SET_PROPERTY_VALUE(Literal, InLiteral))
@@ -348,6 +424,20 @@ void UMetaSoundLiteralViewModel_FloatArray::SetSourceValue(const TArray<float>&
 	}
 }
 
+#if WITH_EDITOR
+void UMetaSoundLiteralViewModel_FloatArray::PostEditUndo()
+{
+	Super::PostEditUndo();
+	SetSourceValue(SourceValue);
+}
+#endif
+
+void UMetaSoundLiteralViewModel_String::PostInitProperties()
+{
+	Super::PostInitProperties();
+	SetFlags(RF_Transactional);
+}
+
 void UMetaSoundLiteralViewModel_String::SetLiteral(const FMetasoundFrontendLiteral& InLiteral)
 {
 	if (UE_MVVM_SET_PROPERTY_VALUE(Literal, InLiteral))
@@ -367,6 +457,20 @@ void UMetaSoundLiteralViewModel_String::SetSourceValue(const FString& InSourceVa
 	}
 }
 
+#if WITH_EDITOR
+void UMetaSoundLiteralViewModel_String::PostEditUndo()
+{
+	Super::PostEditUndo();
+	SetSourceValue(SourceValue);
+}
+#endif
+
+void UMetaSoundLiteralViewModel_StringArray::PostInitProperties()
+{
+	Super::PostInitProperties();
+	SetFlags(RF_Transactional);
+}
+
 void UMetaSoundLiteralViewModel_StringArray::SetLiteral(const FMetasoundFrontendLiteral& InLiteral)
 {
 	if (UE_MVVM_SET_PROPERTY_VALUE(Literal, InLiteral))
@@ -386,6 +490,20 @@ void UMetaSoundLiteralViewModel_StringArray::SetSourceValue(const TArray<FString
 	}
 }
 
+#if WITH_EDITOR
+void UMetaSoundLiteralViewModel_StringArray::PostEditUndo()
+{
+	Super::PostEditUndo();
+	SetSourceValue(SourceValue);
+}
+#endif
+
+void UMetaSoundLiteralViewModel_Object::PostInitProperties()
+{
+	Super::PostInitProperties();
+	SetFlags(RF_Transactional);
+}
+
 void UMetaSoundLiteralViewModel_Object::SetLiteral(const FMetasoundFrontendLiteral& InLiteral)
 {
 	if (UE_MVVM_SET_PROPERTY_VALUE(Literal, InLiteral))
@@ -407,6 +525,20 @@ void UMetaSoundLiteralViewModel_Object::SetSourceValue(UObject* InSourceValue)
 	}
 }
 
+#if WITH_EDITOR
+void UMetaSoundLiteralViewModel_Object::PostEditUndo()
+{
+	Super::PostEditUndo();
+	SetSourceValue(SourceValue);
+}
+#endif
+
+void UMetaSoundLiteralViewModel_ObjectArray::PostInitProperties()
+{
+	Super::PostInitProperties();
+	SetFlags(RF_Transactional);
+}
+
 void UMetaSoundLiteralViewModel_ObjectArray::SetLiteral(const FMetasoundFrontendLiteral& InLiteral)
 {
 	if (UE_MVVM_SET_PROPERTY_VALUE(Literal, InLiteral))
@@ -427,3 +559,11 @@ void UMetaSoundLiteralViewModel_ObjectArray::SetSourceValue(const TArray<TObject
 		UE_MVVM_BROADCAST_FIELD_VALUE_CHANGED(Literal);
 	}
 }
+
+#if WITH_EDITOR
+void UMetaSoundLiteralViewModel_ObjectArray::PostEditUndo()
+{
+	Super::PostEditUndo();
+	SetSourceValue(SourceValue);
+}
+#endif
diff --git a/Engine/Plugins/Experimental/TechAudioTools/Source/TechAudioToolsMetaSound/Public/ViewModels/MetaSoundLiteralViewModel.h b/Engine/Plugins/Experimental/TechAudioTools/Source/TechAudioToolsMetaSound/Public/ViewModels/MetaSoundLiteralViewModel.h
--- a/Engine/Plugins/Experimental/TechAudioTools/Source/TechAudioToolsMetaSound/Public/ViewModels/MetaSoundLiteralViewModel.h
+++ b/Engine/Plugins/Experimental/TechAudioTools/Source/TechAudioToolsMetaSound/Public/ViewModels/MetaSoundLiteralViewModel.h
@@ -36,6 +36,12 @@ public:
 
 	UE_API void SetLiteral(const FMetasoundFrontendLiteral& InLiteral);
 	UE_API void SetSourceValue(bool InSourceValue);
+
+	virtual void PostInitProperties() override;
+
+#if WITH_EDITOR
+	virtual void PostEditUndo() override;
+#endif
 };
 
 /**
@@ -61,6 +67,12 @@ public:
 
 	UE_API void SetLiteral(const FMetasoundFrontendLiteral& InLiteral);
 	UE_API void SetSourceValue(const TArray<bool>& InSourceValue);
+
+	virtual void PostInitProperties() override;
+
+#if WITH_EDITOR
+	virtual void PostEditUndo() override;
+#endif
 };
 
 /**
@@ -100,6 +112,12 @@ public:
 	UE_API void SetSourceValue(int32 InSourceValue);
 	UE_API void SetNormalizedValue(float InNormalizedValue);
 
+	virtual void PostInitProperties() override;
+
+#if WITH_EDITOR
+	virtual void PostEditUndo() override;
+#endif
+
 private:
 	bool bIsUpdating = false;
 };
@@ -127,6 +145,12 @@ public:
 
 	UE_API void SetLiteral(const FMetasoundFrontendLiteral& InLiteral);
 	UE_API void SetSourceValue(const TArray<int32>& InSourceValue);
+
+	virtual void PostInitProperties() override;
+
+#if WITH_EDITOR
+	virtual void PostEditUndo() override;
+#endif
 };
 
 /**
@@ -217,6 +241,12 @@ public:
 	UE_API void SetNormalizedValue(float InNormalizedValue);
 	UE_API void SetDisplayValue(float InDisplayValue);
 
+	virtual void PostInitProperties() override;
+
+#if WITH_EDITOR
+	virtual void PostEditUndo() override;
+#endif
+
 private:
 	bool bIsUpdating = false;
 };
@@ -244,6 +274,12 @@ public:
 
 	UE_API void SetLiteral(const FMetasoundFrontendLiteral& InLiteral);
 	UE_API void SetSourceValue(const TArray<float>& InSourceValue);
+
+	virtual void PostInitProperties() override;
+
+#if WITH_EDITOR
+	virtual void PostEditUndo() override;
+#endif
 };
 
 /**
@@ -269,6 +305,12 @@ public:
 
 	UE_API void SetLiteral(const FMetasoundFrontendLiteral& InLiteral);
 	UE_API void SetSourceValue(const FString& InSourceValue);
+
+	virtual void PostInitProperties() override;
+
+#if WITH_EDITOR
+	virtual void PostEditUndo() override;
+#endif
 };
 
 /**
@@ -294,6 +336,12 @@ public:
 
 	UE_API void SetLiteral(const FMetasoundFrontendLiteral& InLiteral);
 	UE_API void SetSourceValue(const TArray<FString>& InSourceValue);
+
+	virtual void PostInitProperties() override;
+
+#if WITH_EDITOR
+	virtual void PostEditUndo() override;
+#endif
 };
 
 /**
@@ -319,6 +367,12 @@ public:
 
 	UE_API void SetLiteral(const FMetasoundFrontendLiteral& InLiteral);
 	UE_API void SetSourceValue(UObject* InSourceValue);
+
+	virtual void PostInitProperties() override;
+
+#if WITH_EDITOR
+	virtual void PostEditUndo() override;
+#endif
 };
 
 /**
@@ -344,6 +398,12 @@ public:
 
 	UE_API void SetLiteral(const FMetasoundFrontendLiteral& InLiteral);
 	UE_API void SetSourceValue(const TArray<TObjectPtr<UObject>>& InSourceValue);
+
+	virtual void PostInitProperties() override;
+
+#if WITH_EDITOR
+	virtual void PostEditUndo() override;
+#endif
 };
 
 #undef UE_API
